SchematicGraphPanel: Move code to editor module

#rnx

diff --git a/Engine/Plugins/Animation/ControlRig/Source/ControlRigEditor/ControlRigEditor.Build.cs b/Engine/Plugins/Animation/ControlRig/Source/ControlRigEditor/ControlRigEditor.Build.cs
--- a/Engine/Plugins/Animation/ControlRig/Source/ControlRigEditor/ControlRigEditor.Build.cs
+++ b/Engine/Plugins/Animation/ControlRig/Source/ControlRigEditor/ControlRigEditor.Build.cs
@@ -77,6 +77,7 @@ namespace UnrealBuildTool.Rules
 					"AdvancedPreviewScene",
 					"ToolWidgets",
                     "AnimationWidgets",
+                    "AnimationEditorWidgets",
                     "ActorPickerMode",
                     "Constraints",
                     "AnimationEditMode"
diff --git a/Engine/Source/Developer/AnimationWidgets/Private/SchematicGraphPanel/SSchematicGraphPanel.cpp b/Engine/Source/Developer/AnimationWidgets/Private/SchematicGraphPanel/SSchematicGraphPanel.cpp
deleted file mode 100644
diff --git a/Engine/Source/Developer/AnimationWidgets/Private/SchematicGraphPanel/SchematicViewportTabSummoner.cpp b/Engine/Source/Developer/AnimationWidgets/Private/SchematicGraphPanel/SchematicViewportTabSummoner.cpp
deleted file mode 100644
diff --git a/Engine/Source/Developer/AnimationWidgets/Public/SchematicGraphPanel/SSchematicGraphPanel.h b/Engine/Source/Developer/AnimationWidgets/Public/SchematicGraphPanel/SSchematicGraphPanel.h
deleted file mode 100644
diff --git a/Engine/Source/Developer/AnimationWidgets/Public/SchematicGraphPanel/SchematicViewportTabSummoner.h b/Engine/Source/Developer/AnimationWidgets/Public/SchematicGraphPanel/SchematicViewportTabSummoner.h
deleted file mode 100644
diff --git a/Engine/Source/Editor/AnimationEditorWidgets/AnimationEditorWidgets.Build.cs b/Engine/Source/Editor/AnimationEditorWidgets/AnimationEditorWidgets.Build.cs
new file mode 100644
--- /dev/null
+++ b/Engine/Source/Editor/AnimationEditorWidgets/AnimationEditorWidgets.Build.cs
@@ -0,0 +1,34 @@
+// Copyright Epic Games, Inc. All Rights Reserved.
+
+using UnrealBuildTool;
+
+public class AnimationEditorWidgets : ModuleRules
+{
+	public AnimationEditorWidgets(ReadOnlyTargetRules Target) : base(Target)
+	{
+		/** NOTE: THIS MODULE SHOULD NOT EVER DEPEND ON UNREALED. 
+		 * Please refer to ToolWidgets.Build.cs for more information.
+		 */
+		PublicDependencyModuleNames.AddRange(
+			new string[] {
+				"Core",
+				"Slate",
+				"SlateCore",
+				"InputCore",
+				"ToolWidgets",
+				"AnimationCore",
+			}
+		);
+
+		PrivateDependencyModuleNames.AddRange(
+			new string[]
+			{
+				"Engine",
+				"CoreUObject",
+				"UnrealEd",
+				"GraphEditor",
+				"PropertyEditor"
+			}
+		);
+	}
+}
diff --git a/Engine/Source/Editor/AnimationEditorWidgets/Private/AnimationEditorWidgetsModule.cpp b/Engine/Source/Editor/AnimationEditorWidgets/Private/AnimationEditorWidgetsModule.cpp
new file mode 100644
--- /dev/null
+++ b/Engine/Source/Editor/AnimationEditorWidgets/Private/AnimationEditorWidgetsModule.cpp
@@ -0,0 +1,18 @@
+// Copyright Epic Games, Inc. All Rights Reserved.
+
+#include "Modules/ModuleInterface.h"
+#include "Modules/ModuleManager.h"
+
+class FAnimationEditorWidgetsModule
+	: public IModuleInterface
+{
+public:
+
+	// IModuleInterface interface
+	
+	virtual void StartupModule() override { }
+	virtual void ShutdownModule() override { }
+};
+
+
+IMPLEMENT_MODULE(FAnimationEditorWidgetsModule, AnimationWidgets);
diff --git a/Engine/Source/Editor/AnimationEditorWidgets/Private/SchematicGraphPanel/SSchematicGraphPanel.cpp b/Engine/Source/Editor/AnimationEditorWidgets/Private/SchematicGraphPanel/SSchematicGraphPanel.cpp
new file mode 100644
--- /dev/null
+++ b/Engine/Source/Editor/AnimationEditorWidgets/Private/SchematicGraphPanel/SSchematicGraphPanel.cpp
@@ -0,0 +1,869 @@
+// Copyright Epic Games, Inc. All Rights Reserved.
+#if WITH_EDITOR
+
+#include "SchematicGraphPanel/SSchematicGraphPanel.h"
+#include "Framework/Application/SlateApplication.h"
+#include "DragAndDrop/AssetDragDropOp.h"
+
+#define LOCTEXT_NAMESPACE "SSchematicGraphPanel"
+
+const FName& FSchematicGraphNode::GetType() const
+{
+	return FSchematicGraphNode::Type;
+}
+
+bool FSchematicGraphNode::IsA(const FName& InType) const
+{
+	return Type == InType;
+}
+
+FString FSchematicGraphNode::GetDragDropDecoratorLabel() const
+{
+	return GetGuid().ToString();
+}
+
+void FSchematicGraphModel::Reset()
+{
+	Nodes.Reset();
+	Links.Reset();
+
+	if (OnGraphResetDelegate.IsBound())
+	{
+		OnGraphReset().Broadcast();
+	}
+}
+
+bool FSchematicGraphModel::RemoveNode(const FGuid& InGuid)
+{
+	if(const FSchematicGraphNode* Node = FindNode(InGuid))
+	{
+		if (OnNodeRemovedDelegate.IsBound())
+		{
+			OnNodeRemovedDelegate.Broadcast(Node);
+		}
+		const FGuid Guid = Node->GetGuid();
+		NodeByGuid.Remove(Node->GetGuid());
+		Nodes.RemoveAll([Guid](const TSharedPtr<FSchematicGraphNode>& ExistingNode) -> bool
+		{
+			return ExistingNode->GetGuid() == Guid;
+		});
+		return true;
+	}
+	return false;
+}
+
+FVector2d FSchematicGraphModel::GetPositionForNode(const FGuid& InGuid) const
+{
+	if(const FSchematicGraphNode* Node = FindNode(InGuid))
+	{
+		return GetPositionForNode(Node);
+	}
+	return FVector2d::ZeroVector;
+}
+
+FVector2d FSchematicGraphModel::GetPositionForNode(const FSchematicGraphNode* InNode) const
+{
+	check(InNode);
+	return InNode->GetPosition();
+}
+
+FVector2d FSchematicGraphModel::GetPositionOffsetForNode(const FGuid& InGuid) const
+{
+	if(const FSchematicGraphNode* Node = FindNode(InGuid))
+	{
+		return GetPositionOffsetForNode(Node);
+	}
+	return FVector2d::ZeroVector;
+}
+
+FVector2d FSchematicGraphModel::GetPositionOffsetForNode(const FSchematicGraphNode* InNode) const
+{
+	check(InNode);
+	return InNode->GetPositionOffset();
+}
+
+bool FSchematicGraphModel::GetPositionAnimationEnabledForNode(const FGuid& InNodeGuid) const
+{
+	if(const FSchematicGraphNode* Node = FindNode(InNodeGuid))
+	{
+		return GetPositionAnimationEnabledForNode(Node);
+	}
+	return false;
+}
+
+bool FSchematicGraphModel::GetPositionAnimationEnabledForNode(const FSchematicGraphNode* InNode) const
+{
+	return true;
+}
+
+FVector2d FSchematicGraphModel::GetSizeForNode(const FGuid& InNodeGuid) const
+{
+	if(const FSchematicGraphNode* Node = FindNode(InNodeGuid))
+	{
+		return GetSizeForNode(Node);
+	}
+	return FVector2d::ZeroVector;
+}
+
+FVector2d FSchematicGraphModel::GetSizeForNode(const FSchematicGraphNode* InNode) const
+{
+	static const FVector2d DefaultSize = FVector2d(50.0,50.0);
+	return DefaultSize;
+}
+
+float FSchematicGraphModel::GetScaleForNode(const FGuid& InNodeGuid) const
+{
+	if(const FSchematicGraphNode* Node = FindNode(InNodeGuid))
+	{
+		return GetScaleForNode(Node);
+	}
+	return 1.f;
+}
+
+float FSchematicGraphModel::GetScaleForNode(const FSchematicGraphNode* InNode) const
+{
+	return 1.f;
+}
+
+FLinearColor FSchematicGraphModel::GetColorForNode(const FGuid& InGuid) const
+{
+	if(const FSchematicGraphNode* Node = FindNode(InGuid))
+	{
+		const FLinearColor Color = GetColorForNode(Node);
+		if(GetVisibilityForNode(Node) == ESchematicGraphNodeVisibility::FadedOut)
+		{
+			return Color * 0.5f;
+		}
+		return Color;
+	}
+	return FLinearColor::White;
+}
+
+FLinearColor FSchematicGraphModel::GetColorForNode(const FSchematicGraphNode* InNode) const
+{
+	check(InNode);
+	return InNode->GetColor();
+}
+
+const FSlateBrush* FSchematicGraphModel::GetBrushForNode(const FGuid& InGuid) const
+{
+	if(const FSchematicGraphNode* Node = FindNode(InGuid))
+	{
+		if(const FSlateBrush* Brush = GetBrushForNode(Node))
+		{
+			return Brush;
+		}
+	}
+
+	static const FSlateBrush* DefaultBrush = SSchematicGraphNode::FArguments()._Brush.Get();
+	return DefaultBrush;
+}
+
+const FSlateBrush* FSchematicGraphModel::GetBrushForNode(const FSchematicGraphNode* InNode) const
+{
+	check(InNode);
+	return InNode->GetBrush();
+}
+
+const FText FSchematicGraphModel::GetToolTipForNode(const FGuid& InGuid) const
+{
+	if(const FSchematicGraphNode* Node = FindNode(InGuid))
+	{
+		return GetToolTipForNode(Node);
+	}
+	return FText();
+}
+
+const FText FSchematicGraphModel::GetToolTipForNode(const FSchematicGraphNode* InNode) const
+{
+	check(InNode);
+	return InNode->GetToolTip();
+}
+
+ESchematicGraphNodePlacementConstraint FSchematicGraphModel::GetPlacementForNode(const FGuid& InGuid) const
+{
+	if(const FSchematicGraphNode* Node = FindNode(InGuid))
+	{
+		return GetPlacementForNode(Node);
+	}
+	return ESchematicGraphNodePlacementConstraint::Free;
+}
+
+ESchematicGraphNodePlacementConstraint FSchematicGraphModel::GetPlacementForNode(const FSchematicGraphNode* InNode) const
+{
+	check(InNode);
+	return InNode->GetPlacement();
+}
+
+ESchematicGraphNodeVisibility FSchematicGraphModel::GetVisibilityForNode(const FGuid& InGuid) const
+{
+	if(const FSchematicGraphNode* Node = FindNode(InGuid))
+	{
+		return GetVisibilityForNode(Node);
+	}
+	return ESchematicGraphNodeVisibility::Visible;
+}
+
+ESchematicGraphNodeVisibility FSchematicGraphModel::GetVisibilityForNode(const FSchematicGraphNode* InNode) const
+{
+	check(InNode);
+	return InNode->GetVisibility();
+}
+
+TSharedRef<FSchematicGraphNodeDragDropOp> FSchematicGraphNodeDragDropOp::New(TArray<SSchematicGraphNode*> InSchematicGraphNodes, const TArray<FGuid>& InElements, FSchematicGraphNodeDragDropOp::FOnEndDrag InOnEndDragDelegate)
+{
+	TSharedRef<FSchematicGraphNodeDragDropOp> Operation = MakeShared<FSchematicGraphNodeDragDropOp>();
+	Operation->SchematicGraphNodes = InSchematicGraphNodes;
+	Operation->Elements = InElements;
+	Operation->OnEndDragDelegate = InOnEndDragDelegate;
+	Operation->Construct();
+	return Operation;
+}
+
+FSchematicGraphNodeDragDropOp::~FSchematicGraphNodeDragDropOp()
+{
+	if (OnEndDragDelegate.IsBound())
+	{
+		for (SSchematicGraphNode* Node : SchematicGraphNodes)
+		{
+			OnEndDragDelegate.Execute(Node, *this);
+		}
+	}
+}
+
+TSharedPtr<SWidget> FSchematicGraphNodeDragDropOp::GetDefaultDecorator() const
+{
+	return SNew(SBorder)
+		.Visibility(EVisibility::Visible)
+		.BorderImage(FAppStyle::GetBrush("Menu.Background"))
+		[
+			SNew(STextBlock)
+			.Text(FText::FromString(GetJoinedDecoratorLabels()))
+		];
+}
+
+FString FSchematicGraphNodeDragDropOp::GetJoinedDecoratorLabels() const
+{
+	TArray<FString> DecoratorLabels;
+	for(const SSchematicGraphNode* GraphNode : SchematicGraphNodes)
+	{
+		if(const FSchematicGraphNode* Node = GraphNode->GetNodeData())
+		{
+			DecoratorLabels.Add(Node->GetDragDropDecoratorLabel());
+		}
+	}
+	return FString::Join(DecoratorLabels, TEXT(","));
+}
+
+SSchematicGraphNode::FArguments::FArguments()
+: _NodeData(nullptr)
+, _Brush(nullptr)
+{
+	static const FSchematicGraphNode EmptyNodeData;
+	_NodeData = &EmptyNodeData;
+	static const FSlateBrush* WhiteTexture = FAppStyle::GetBrush("WhiteTexture");
+	_Brush = WhiteTexture;
+}
+
+void SSchematicGraphNode::Construct(const FArguments& InArgs)
+{
+	if(InArgs._NodeData)
+	{
+		NodeData = InArgs._NodeData;
+	}
+	OnClickedDelegate = InArgs._OnClicked;
+	OnBeginDragDelegate = InArgs._OnBeginDrag;
+	OnEndDragDelegate = InArgs._OnEndDrag;
+	OnDropDelegate = InArgs._OnDrop;
+
+	Position = InArgs._Position;
+	Size = InArgs._Size;
+	Scale = InArgs._Scale;
+	Color = InArgs._Color;
+	if(InArgs._Brush.IsSet() || InArgs._Brush.IsBound())
+	{
+		Brush = InArgs._Brush;
+	}
+	OriginalSize = Size->Get();
+
+	if(InArgs._ToolTipText.IsBound() || InArgs._ToolTipText.IsSet())
+	{
+		SetToolTipText(InArgs._ToolTipText);
+	}
+}
+
+FVector2D SSchematicGraphNode::ComputeDesiredSize(float LayoutScaleMultiplier) const
+{
+	return OriginalSize*LayoutScaleMultiplier;
+}
+
+int32 SSchematicGraphNode::OnPaint(const FPaintArgs& Args, const FGeometry& AllottedGeometry, const FSlateRect& MyCullingRect, FSlateWindowElementList& OutDrawElements, int32 LayerId, const FWidgetStyle& InWidgetStyle, bool bParentEnabled) const
+{
+	int32 NewLayerId = SNodePanel::SNode::OnPaint(Args, AllottedGeometry, MyCullingRect, OutDrawElements, LayerId, InWidgetStyle, bParentEnabled);
+	NewLayerId++;
+
+	const FVector2d CurSize = Size->Get() * Scale->Get();
+	const FVector2d SizeOffset = (CurSize-OriginalSize)*-0.5;
+
+	FSlateDrawElement::MakeBox(
+				OutDrawElements,
+				NewLayerId,
+				AllottedGeometry.ToPaintGeometry(CurSize, FSlateLayoutTransform(SizeOffset)),
+				Brush.Get(),
+				ESlateDrawEffect::None,
+				Color.IsValid() ? Color->Get() : FLinearColor::White
+				);
+	
+	return NewLayerId;
+}
+
+void SSchematicGraphNode::OnMouseEnter(const FGeometry& MyGeometry, const FPointerEvent& MouseEvent)
+{
+	if(!IsInteractive())
+	{
+		return;
+	}
+	SNode::OnMouseEnter(MyGeometry, MouseEvent);
+	Scale->Set(ScaledUp);
+}
+
+void SSchematicGraphNode::OnMouseLeave(const FPointerEvent& MouseEvent)
+{
+	if(!IsInteractive())
+	{
+		return;
+	}
+	SNode::OnMouseLeave(MouseEvent);
+	Scale->Set(1.f);
+}
+
+FReply SSchematicGraphNode::OnDragOver(const FGeometry& MyGeometry, const FDragDropEvent& DragDropEvent)
+{
+	if(!IsInteractive())
+	{
+		return FReply::Unhandled();
+	}
+	SNode::OnDragOver(MyGeometry, DragDropEvent);
+	Scale->Set(ScaledDown);
+	return FReply::Handled();
+}
+
+void SSchematicGraphNode::OnDragLeave(const FDragDropEvent& DragDropEvent)
+{
+	if(!IsInteractive())
+	{
+		return;
+	}
+	SNode::OnDragLeave(DragDropEvent);
+	Scale->Set(1.f);
+}
+
+FReply SSchematicGraphNode::OnDrop(const FGeometry& MyGeometry, const FDragDropEvent& DragDropEvent)
+{
+	if(!IsInteractive())
+	{
+		return FReply::Unhandled();
+	}
+
+	// Avoid dropping onto itself
+	const TSharedPtr<FSchematicGraphNodeDragDropOp> SchematicDragDropOp = DragDropEvent.GetOperationAs<FSchematicGraphNodeDragDropOp>();
+	if (SchematicDragDropOp)
+	{
+		if (!SchematicDragDropOp->GetElements().IsEmpty())
+		{
+			if (SchematicDragDropOp->GetElements().ContainsByPredicate([this](const FGuid& Guid)
+			{
+				return Guid == NodeData->GetGuid();
+			}))
+			{
+				return FReply::Unhandled();
+			}
+		}
+	}
+	
+	SNode::OnDrop(MyGeometry, DragDropEvent);
+	OnDropDelegate.ExecuteIfBound(this, DragDropEvent);
+	OnEndDragDelegate.ExecuteIfBound(this, *DragDropEvent.GetOperation().Get());
+	return FReply::Unhandled();
+}
+
+FReply SSchematicGraphNode::OnDragDetected(const FGeometry& MyGeometry, const FPointerEvent& MouseEvent)
+{
+	TArray<FGuid> DraggedElements = {GetGuid()};
+	if (MouseEvent.IsMouseButtonDown(EKeys::LeftMouseButton) && DraggedElements.Num() > 0)
+	{
+		bIsBeingDragged = true;
+
+		const FVector2f AbsoluteMousePosition = MouseEvent.GetScreenSpacePosition();
+		const FVector2d LocalMousePosition =  (AbsoluteMousePosition - MyGeometry.GetAbsolutePosition()) / MyGeometry.GetAccumulatedLayoutTransform().GetScale();
+		OffsetDuringDrag = -LocalMousePosition;
+		
+		const TSharedRef<FSchematicGraphNodeDragDropOp> DragDropOp = FSchematicGraphNodeDragDropOp::New({this}, MoveTemp(DraggedElements), OnEndDragDelegate);
+		OnBeginDragDelegate.ExecuteIfBound(this, DragDropOp.Get());
+		return FReply::Handled().BeginDragDrop(DragDropOp);
+	}
+	
+	return FReply::Unhandled();
+}
+
+FReply SSchematicGraphNode::OnMouseButtonDown(const FGeometry& MyGeometry, const FPointerEvent& MouseEvent)
+{
+	SNode::OnMouseButtonDown(MyGeometry, MouseEvent);
+	if (MouseEvent.GetPressedButtons().Contains(EKeys::LeftMouseButton))
+	{
+		OnClickedDelegate.ExecuteIfBound(this);
+	}
+	return FReply::Handled().DetectDrag(SharedThis(this), EKeys::LeftMouseButton);
+}
+
+FVector2d SSchematicGraphNode::GetPosition() const
+{
+	return Position->Get() - (OriginalSize*0.5);
+}
+
+void SSchematicGraphNode::EnablePositionAnimation(bool bEnabled)
+{
+	Position->EnableInterpolation(bEnabled);
+}
+
+const FGuid SSchematicGraphNode::GetGuid() const
+{
+	return GetNodeData()->GetGuid();
+}
+
+bool SSchematicGraphNode::IsInteractive() const
+{
+	if(NodeData)
+	{
+		return NodeData->GetVisibility() == ESchematicGraphNodeVisibility::Visible;
+	}
+	return true;
+}
+
+void SSchematicGraphPanel::SetSchematicGraph(FSchematicGraphModel* InGraphData)
+{
+	if (GraphData)
+	{
+		GraphData->OnNodeAdded().RemoveAll(this);
+		GraphData->OnNodeRemoved().RemoveAll(this);
+		GraphData->OnGraphReset().RemoveAll(this);
+	}
+	
+	GraphData = InGraphData;
+	if (GraphData)
+	{
+		GraphData->OnNodeAdded().AddSP(this, &SSchematicGraphPanel::AddNode);
+		GraphData->OnNodeRemoved().AddSP(this, &SSchematicGraphPanel::RemoveNode);
+		GraphData->OnGraphReset().AddSP(this, &SSchematicGraphPanel::RebuildPanel);
+	}
+}
+
+void SSchematicGraphPanel::Construct(const FArguments& InArgs)
+{
+	GraphData = InArgs._GraphData;
+	bIsOverlay = InArgs._IsOverlay;
+	PaddingLeft = InArgs._PaddingLeft;
+	PaddingRight = InArgs._PaddingRight;
+	PaddingTop = InArgs._PaddingTop;
+	PaddingBottom = InArgs._PaddingBottom;
+	PaddingInterNode = InArgs._PaddingInterNode;
+	OnNodeClickedDelegate = InArgs._OnNodeClicked;
+	OnBeginDragDelegate = InArgs._OnBeginDrag;
+	OnEndDragDelegate = InArgs._OnEndDrag;
+	OnDropDelegate = InArgs._OnDrop;
+
+	TEasingAttributeInterpolator<float>::FSettings FloatInterpSettings(EEasingInterpolatorType::CubicEaseOut, 0.35f);
+	FadeBackgroundAlpha = FFloatAttribute::CreateWithGetter(FloatInterpSettings, FFloatAttribute::FGetter::CreateLambda([this]() -> float
+	{
+		return bIsDragDropping ? 0.5f : 0.f;
+	}));
+	
+	SNodePanel::Construct();
+
+	SetVisibility(bIsOverlay ? EVisibility::SelfHitTestInvisible : EVisibility::Visible);
+	if (GraphData)
+	{
+		GraphData->OnNodeAdded().AddSP(this, &SSchematicGraphPanel::AddNode);
+		GraphData->OnNodeRemoved().AddSP(this, &SSchematicGraphPanel::RemoveNode);
+		GraphData->OnGraphReset().AddSP(this, &SSchematicGraphPanel::RebuildPanel);
+	}
+}
+
+void SSchematicGraphPanel::RebuildPanel()
+{
+	RemoveAllNodes();
+
+	if (GraphData)
+	{
+		for (const TSharedPtr<FSchematicGraphNode>& Node : GraphData->GetNodes())
+		{
+			AddNode(Node.Get());
+		}
+	}
+}
+
+void SSchematicGraphPanel::AddNode(const FSchematicGraphNode* InNodeToAdd)
+{
+	TEasingAttributeInterpolator<FVector2d>::FSettings Vector2DInterpSettings(EEasingInterpolatorType::CubicEaseOut, 0.2f);
+	TEasingAttributeInterpolator<float>::FSettings FloatInterpSettings(EEasingInterpolatorType::CubicEaseOut, 0.2f);
+	TEasingAttributeInterpolator<FLinearColor>::FSettings ColorInterpSettings(EEasingInterpolatorType::CubicEaseOut, 0.2f);
+
+	const FGuid Guid = InNodeToAdd->GetGuid();
+	const auto Position = FVector2dAttribute::CreateWithGetter(Vector2DInterpSettings, FVector2dAttribute::FGetter::CreateSP(this, &SSchematicGraphPanel::GetPositionForNode, Guid));
+	const auto Size = FVector2dAttribute::CreateWithGetter(Vector2DInterpSettings, FVector2dAttribute::FGetter::CreateSP(this, &SSchematicGraphPanel::GetSizeForNode, Guid));
+	const auto Scale = FFloatAttribute::CreateWithGetter(FloatInterpSettings, FFloatAttribute::FGetter::CreateSP(this, &SSchematicGraphPanel::GetScaleForNode, Guid), 0.f);
+	const auto Color = FLinearColorAttribute::CreateWithGetter(ColorInterpSettings, FLinearColorAttribute::FGetter::CreateSP(this, &SSchematicGraphPanel::GetColorForNode, Guid));
+
+	const TSharedRef<SSchematicGraphNode> NewNode = SNew(SSchematicGraphNode)
+														.Position(Position)
+														.Size(Size)
+														.Scale(Scale)
+														.Color(Color)
+														.Brush(this, &SSchematicGraphPanel::GetBrushForNode, InNodeToAdd->GetGuid())
+														.ToolTipText(this, &SSchematicGraphPanel::GetToolTipForNode, InNodeToAdd->GetGuid())
+														.OnClicked_Raw(this, &SSchematicGraphPanel::OnNodeClicked)
+														.OnBeginDrag_Raw(this, &SSchematicGraphPanel::OnBeginDragEvent)
+														.OnEndDrag_Raw(this, &SSchematicGraphPanel::OnEndDragEvent)
+														.OnDrop_Raw(this, &SSchematicGraphPanel::OnDropEvent)
+														.NodeData(InNodeToAdd);
+	SNodePanel::AddGraphNode(NewNode);
+	NodeByGuid.Add(Guid, NewNode.ToSharedPtr());
+}
+
+void SSchematicGraphPanel::RemoveNode(const FSchematicGraphNode* InNodeToRemove)
+{
+	const FGuid GuidToRemove = InNodeToRemove->GetGuid();
+	NodeByGuid.Remove(GuidToRemove);
+	
+	for (int32 Iter = 0; Iter != Children.Num(); ++Iter)
+	{
+		TSharedRef<SSchematicGraphNode> Child = GetChild(Iter);
+		if (Child->GetGuid() == GuidToRemove)
+		{
+			Children.RemoveAt(Iter);
+			break;
+		}
+	}
+	for (int32 Iter = 0; Iter != VisibleChildren.Num(); ++Iter)
+	{
+		TSharedRef<SSchematicGraphNode> Child = StaticCastSharedRef<SSchematicGraphNode>(VisibleChildren[Iter]);
+		if (Child->GetGuid() == GuidToRemove)
+		{
+			VisibleChildren.RemoveAt(Iter);
+			break;
+		}
+	}
+}
+
+const SSchematicGraphNode* SSchematicGraphPanel::FindNode(const FGuid& InGuid) const
+{
+	if(const TSharedPtr<SSchematicGraphNode>* FoundNodePtr = NodeByGuid.Find(InGuid))
+	{
+		return FoundNodePtr->Get();
+	}
+	return nullptr;
+}
+
+void SSchematicGraphPanel::OnArrangeChildren(const FGeometry& AllottedGeometry, FArrangedChildren& ArrangedChildren) const
+{
+	SNodePanel::OnArrangeChildren(AllottedGeometry, ArrangedChildren);
+}
+
+int32 SSchematicGraphPanel::OnPaint(const FPaintArgs& Args, const FGeometry& AllottedGeometry, const FSlateRect& MyCullingRect, FSlateWindowElementList& OutDrawElements, int32 LayerId, const FWidgetStyle& InWidgetStyle, bool bParentEnabled) const
+{
+	const int32 NodeLayerId = LayerId + 1;
+    int32 MaxLayerId = NodeLayerId;
+    	
+	if (!bIsOverlay)
+	{
+		const FSlateBrush* DefaultBackground = FAppStyle::GetBrush(TEXT("Graph.Panel.SolidBackground"));
+		PaintBackgroundAsLines(DefaultBackground, AllottedGeometry, MyCullingRect, OutDrawElements, LayerId);
+		MaxLayerId++;
+	}
+
+	const float BackgroundAlpha = FadeBackgroundAlpha->Get();
+	if (BackgroundAlpha > 0.f)
+	{
+		const FSlateBrush* Brush = FAppStyle::GetBrush(TEXT("Graph.Panel.SolidBackground"));
+		const FLinearColor TransparentGrey = FLinearColor(0.5, 0.5, 0.5, BackgroundAlpha);
+		FSlateDrawElement::MakeBox(
+				OutDrawElements,
+				MaxLayerId,
+				AllottedGeometry.ToPaintGeometry(),
+				Brush,
+				ESlateDrawEffect::None,
+				TransparentGrey
+				);
+		MaxLayerId++;
+	}
+	
+	
+
+	if (!GraphData)
+	{
+		return MaxLayerId; 
+	}
+
+	FArrangedChildren ArrangedChildren(EVisibility::Visible);
+	ArrangeChildNodes(AllottedGeometry, ArrangedChildren);
+
+	// Because we paint multiple children, we must track the maximum layer id that they produced in case one of our parents
+	// wants to an overlay for all of its contents.
+
+	const FPaintArgs NewArgs = Args.WithNewParent(this);
+
+	// Draw the child nodes
+	{
+		for (int32 ChildIndex = 0; ChildIndex < ArrangedChildren.Num(); ++ChildIndex)
+		{
+			FArrangedWidget& CurWidget = ArrangedChildren[ChildIndex];
+			TSharedRef<SSchematicGraphNode> ChildNode = StaticCastSharedRef<SSchematicGraphNode>(CurWidget.Widget);
+			
+			// Examine node to see what layers we should be drawing in
+			const int32 ChildLayerId = NodeLayerId;
+
+			const bool bNodeIsVisible = FSlateRect::DoRectanglesIntersect( CurWidget.Geometry.GetLayoutBoundingRect(), MyCullingRect );
+
+			if (bNodeIsVisible)
+			{
+				const int32 CurWidgetsMaxLayerId = CurWidget.Widget->Paint(NewArgs, CurWidget.Geometry, MyCullingRect, OutDrawElements, ChildLayerId, InWidgetStyle, true );
+				MaxLayerId = FMath::Max( MaxLayerId, CurWidgetsMaxLayerId + 1 );
+			}
+		}
+	}
+
+	// Draw the software cursor
+	++MaxLayerId;
+	PaintSoftwareCursor(AllottedGeometry, MyCullingRect, OutDrawElements, MaxLayerId);
+
+	return MaxLayerId;
+}
+
+void SSchematicGraphPanel::RemoveAllNodes()
+{
+	NodeByGuid.Reset();
+	SNodePanel::RemoveAllNodes();
+}
+
+TSharedRef<SSchematicGraphNode> SSchematicGraphPanel::GetChild(int32 ChildIndex) const
+{
+	return StaticCastSharedRef<SSchematicGraphNode>(Children[ChildIndex]);
+}
+
+TStatId SSchematicGraphPanel::GetStatId() const
+{
+	RETURN_QUICK_DECLARE_CYCLE_STAT(SSchematicGraphPanel, STATGROUP_Tickables);
+}
+
+void SSchematicGraphPanel::Tick(float DeltaTime)
+{
+	bIsDragDropping = false;
+	FSlateApplication& Application = FSlateApplication::Get();
+	if (Application.IsDragDropping())
+	{
+		TSharedPtr<FDragDropOperation> DragDropOp = Application.GetDragDroppingContent();
+		if (DragDropOp.IsValid())
+		{
+			bIsDragDropping = true;
+		}
+	}
+
+	NodesTopLeft.Reset();
+	NodesTopRight.Reset();
+	NodesBottomLeft.Reset();
+	NodesBottomRight.Reset();
+
+	FVector2D Size = GetCachedGeometry().Size;
+	for (int32 i=0; i<Children.Num(); ++i)
+	{
+		TSharedRef<SSchematicGraphNode> Node = GetChild(i);
+
+		// update the animation state of the node
+		Node->EnablePositionAnimation(GraphData->GetPositionAnimationEnabledForNode(Node->GetGuid()));
+
+		// collect the nodes displayed in the corners
+		switch(GraphData->GetPlacementForNode(Node->GetNodeData()))
+		{
+			case ESchematicGraphNodePlacementConstraint::Free:
+			{
+				break;
+			}
+			case ESchematicGraphNodePlacementConstraint::TopLeft:
+			{
+				NodesTopLeft.Add(Node->GetGuid());
+				break;
+			}
+			case ESchematicGraphNodePlacementConstraint::TopRight:
+			{
+				NodesTopRight.Add(Node->GetGuid());
+				break;
+			}
+			case ESchematicGraphNodePlacementConstraint::BottomLeft:
+			{
+				NodesBottomLeft.Add(Node->GetGuid());
+				break;
+			}
+			case ESchematicGraphNodePlacementConstraint::BottomRight:
+			{
+				NodesBottomRight.Add(Node->GetGuid());
+				break;
+			}
+		}
+
+		if (Node->IsBeingDragged())
+		{
+			if (bIsDragDropping)
+			{
+				const FVector2f AbsoluteMousePosition = FSlateApplication::Get().GetCursorPos();
+				const FGeometry& Geometry = GetTickSpaceGeometry();
+
+				const FVector2d LocalMousePosition =  (AbsoluteMousePosition - Geometry.GetAbsolutePosition()) / Geometry.GetAccumulatedLayoutTransform().GetScale();
+				const FVector2d HalfOriginalSize = Node->GetOriginalSize() * 0.5;
+				Node->PositionDuringDrag = LocalMousePosition + HalfOriginalSize;
+			}
+			else if (DeltaTime > 0.f)
+			{
+				Node->PositionDuringDrag.Reset();
+				Node->OffsetDuringDrag.Reset();
+				Node->bIsBeingDragged = false;
+			}
+		}
+	}
+}
+
+void SSchematicGraphPanel::Tick(const FGeometry& AllottedGeometry, const double InCurrentTime, const float InDeltaTime)
+{
+	SNodePanel::Tick(AllottedGeometry, InCurrentTime, InDeltaTime);
+	Tick(InDeltaTime);
+}
+
+void SSchematicGraphPanel::OnNodeClicked(SSchematicGraphNode* Node)
+{
+	OnNodeClickedDelegate.ExecuteIfBound(this, Node);
+}
+
+void SSchematicGraphPanel::OnBeginDragEvent(SSchematicGraphNode* Node, const FDragDropOperation& InDragDropOp)
+{
+	OnBeginDragDelegate.ExecuteIfBound(this, Node, InDragDropOp);
+}
+
+void SSchematicGraphPanel::OnEndDragEvent(SSchematicGraphNode* Node, const FDragDropOperation& InDragDropOp)
+{
+	OnEndDragDelegate.ExecuteIfBound(this, Node, InDragDropOp);
+}
+
+void SSchematicGraphPanel::OnDropEvent(SSchematicGraphNode* Node, const FDragDropEvent& InDragDropEvent)
+{
+	OnDropDelegate.ExecuteIfBound(this, Node, InDragDropEvent);
+}
+
+FVector2d SSchematicGraphPanel::GetPositionForNode(FGuid InNodeGuid) const
+{
+	const SSchematicGraphNode* NodeWidget = FindNode(InNodeGuid);
+	if(GraphData && NodeWidget)
+	{
+		if(NodeWidget->PositionDuringDrag.IsSet())
+		{
+			return NodeWidget->PositionDuringDrag.GetValue() + NodeWidget->OffsetDuringDrag.Get(FVector2d::ZeroVector);
+		}
+		
+		if(const FSchematicGraphNode* Node = GraphData->FindNode(InNodeGuid))
+		{
+			FVector2d Position = GraphData->GetPositionOffsetForNode(Node);
+
+			const FVector2d OriginalSize = NodeWidget->GetOriginalSize();
+			const FVector2D Size = GetCachedGeometry().Size;
+			const FVector2d HalfOriginalSize = OriginalSize * 0.5;
+
+			switch(GraphData->GetPlacementForNode(Node))
+			{
+				case ESchematicGraphNodePlacementConstraint::Free:
+				{
+					Position += GraphData->GetPositionForNode(Node);
+					break;
+				}
+				case ESchematicGraphNodePlacementConstraint::TopLeft:
+				{
+					const int32 NodeIndexInCorner = NodesTopLeft.IndexOfByKey(InNodeGuid);
+					check(NodeIndexInCorner != INDEX_NONE);
+					Position += FVector2d(PaddingLeft + HalfOriginalSize.X, PaddingTop + HalfOriginalSize.Y + ((OriginalSize.Y + PaddingInterNode) * NodeIndexInCorner));
+					break;
+				}
+				case ESchematicGraphNodePlacementConstraint::TopRight:
+				{
+					const int32 NodeIndexInCorner = NodesTopRight.IndexOfByKey(InNodeGuid);
+					check(NodeIndexInCorner != INDEX_NONE);
+					Position += FVector2d(Size.X - PaddingRight - HalfOriginalSize.X, PaddingTop + HalfOriginalSize.Y + ((OriginalSize.Y + PaddingInterNode) * NodeIndexInCorner));
+					break;
+				}
+				case ESchematicGraphNodePlacementConstraint::BottomLeft:
+				{
+					const int32 NodeIndexInCorner = NodesBottomLeft.IndexOfByKey(InNodeGuid);
+					check(NodeIndexInCorner != INDEX_NONE);
+					Position += FVector2d(PaddingLeft + HalfOriginalSize.X, Size.Y - PaddingBottom - HalfOriginalSize.Y - ((OriginalSize.Y + PaddingInterNode)  * NodeIndexInCorner));
+					break;
+				}
+				case ESchematicGraphNodePlacementConstraint::BottomRight:
+				{
+					const int32 NodeIndexInCorner = NodesBottomRight.IndexOfByKey(InNodeGuid);
+					check(NodeIndexInCorner != INDEX_NONE);
+					Position += FVector2d(Size.X - PaddingRight - HalfOriginalSize.X, Size.Y - PaddingBottom - HalfOriginalSize.Y - ((OriginalSize.Y + PaddingInterNode)  * NodeIndexInCorner));
+					break;
+				}
+			}
+			return Position;
+		}
+	}
+	return FVector2d::ZeroVector;
+}
+
+FVector2d SSchematicGraphPanel::GetSizeForNode(FGuid InNodeGuid) const
+{
+	if(GraphData)
+	{
+		return GraphData->GetSizeForNode(InNodeGuid);
+	}
+	static const FVector2d DefaultSize = FVector2d(50.0,50.0);
+	return DefaultSize;
+}
+
+float SSchematicGraphPanel::GetScaleForNode(FGuid InNodeGuid) const
+{
+	if(GraphData)
+	{
+		return GraphData->GetScaleForNode(InNodeGuid);
+	}
+	return 1.f;
+}
+
+FLinearColor SSchematicGraphPanel::GetColorForNode(FGuid InNodeGuid) const
+{
+	if(GraphData)
+	{
+		return GraphData->GetColorForNode(InNodeGuid);
+	}
+	return FLinearColor::White;
+}
+
+const FSlateBrush* SSchematicGraphPanel::GetBrushForNode(FGuid InNodeGuid) const
+{
+	if(GraphData)
+	{
+		return GraphData->GetBrushForNode(InNodeGuid);
+	}
+	static const FSlateBrush* DefaultBrush = SSchematicGraphNode::FArguments()._Brush.Get();
+	return DefaultBrush;
+}
+
+FText SSchematicGraphPanel::GetToolTipForNode(FGuid InNodeGuid) const
+{
+	if(GraphData)
+	{
+		return GraphData->GetToolTipForNode(InNodeGuid);
+	}
+	return FText();
+}
+
+#undef LOCTEXT_NAMESPACE
+
+#endif
\ No newline at end of file
diff --git a/Engine/Source/Editor/AnimationEditorWidgets/Private/SchematicGraphPanel/SchematicViewportTabSummoner.cpp b/Engine/Source/Editor/AnimationEditorWidgets/Private/SchematicGraphPanel/SchematicViewportTabSummoner.cpp
new file mode 100644
--- /dev/null
+++ b/Engine/Source/Editor/AnimationEditorWidgets/Private/SchematicGraphPanel/SchematicViewportTabSummoner.cpp
@@ -0,0 +1,52 @@
+// Copyright Epic Games, Inc. All Rights Reserved.
+#if WITH_EDITOR
+
+#include "SchematicGraphPanel/SchematicViewportTabSummoner.h"
+#include "SchematicGraphPanel/SSchematicGraphPanel.h"
+#include "Widgets/Docking/SDockTab.h"
+
+#define LOCTEXT_NAMESPACE "SchematicViewportTabSummoner"
+
+const FName FSchematicViewportTabSummoner::TabID(TEXT("SchematicViewport"));
+
+FSchematicViewportTabSummoner::FSchematicViewportTabSummoner(TSharedPtr<class FAssetEditorToolkit> InHostingApp, const FSchematicViewportArgs& InArgs)
+	: FWorkflowTabFactory(TabID, InHostingApp)
+	, SchematicGraph(InArgs.SchematicGraph)
+	, OnViewportCreated(InArgs.OnViewportCreated)
+{
+	TabLabel = LOCTEXT("SchematicViewportTabLabel", "Schematic Viewport");
+	TabIcon = FSlateIcon(FAppStyle::GetAppStyleSetName(), "LevelEditor.Tabs.Viewports");
+
+	ViewMenuDescription = LOCTEXT("SchematicViewport_ViewMenu_Desc", "Schematic Viewport");
+	ViewMenuTooltip = LOCTEXT("SchematicViewport_ViewMenu_ToolTip", "Show the Schematic Viewport tab");
+}
+
+FTabSpawnerEntry& FSchematicViewportTabSummoner::RegisterTabSpawner(TSharedRef<FTabManager> InTabManager, const FApplicationMode* CurrentApplicationMode) const
+{
+	FTabSpawnerEntry& SpawnerEntry = FWorkflowTabFactory::RegisterTabSpawner(InTabManager, CurrentApplicationMode);
+
+	SpawnerEntry.SetReuseTabMethod(FOnFindTabToReuse::CreateLambda([](const FTabId& InTabId) ->TSharedPtr<SDockTab> {
+	
+		return TSharedPtr<SDockTab>();
+
+	}));
+
+	return SpawnerEntry;
+}
+
+TSharedRef<SWidget> FSchematicViewportTabSummoner::CreateTabBody(const FWorkflowTabSpawnInfo& Info) const
+{
+	TSharedRef<SSchematicGraphPanel> Viewport = SNew(SSchematicGraphPanel)
+											.GraphData(SchematicGraph)
+											.IsOverlay(false)
+											.PaddingLeft(30)
+											.PaddingRight(30)
+											.PaddingTop(30)
+											.PaddingBottom(30)
+											.PaddingInterNode(5);
+	OnViewportCreated.ExecuteIfBound(Viewport);
+	return Viewport;
+}
+
+#undef LOCTEXT_NAMESPACE 
+#endif
\ No newline at end of file
diff --git a/Engine/Source/Editor/AnimationEditorWidgets/Public/SchematicGraphPanel/SSchematicGraphPanel.h b/Engine/Source/Editor/AnimationEditorWidgets/Public/SchematicGraphPanel/SSchematicGraphPanel.h
new file mode 100644
--- /dev/null
+++ b/Engine/Source/Editor/AnimationEditorWidgets/Public/SchematicGraphPanel/SSchematicGraphPanel.h
@@ -0,0 +1,437 @@
+// Copyright Epic Games, Inc. All Rights Reserved.
+#pragma once
+
+#if WITH_EDITOR
+
+#include "CoreMinimal.h"
+#include "SNodePanel.h"
+#include "TickableEditorObject.h"
+#include "Animation/AnimatedAttribute.h"
+
+enum ESchematicGraphNodePlacementConstraint
+{
+	Free,
+	TopLeft,
+	TopRight,
+	BottomLeft,
+	BottomRight
+};
+
+enum ESchematicGraphNodeVisibility
+{
+	Visible,
+	FadedOut,
+	Hidden
+};
+
+#define SCHEMATICGRAPHNODE_BODY(ClassName, SuperClass) \
+inline static const FName& Type = TEXT(#ClassName); \
+virtual const FName& GetType() const override { return ClassName::Type; } \
+virtual bool IsA(const FName& InType) const \
+{ \
+	if(ClassName::Type == InType) \
+	{ \
+		return true; \
+	} \
+	return SuperClass::IsA(InType); \
+} \
+template<typename T> \
+friend const T* Cast(const ClassName* InNode) \
+{ \
+	return Cast<T>((const FSchematicGraphNode*) InNode); \
+} \
+template<typename T> \
+friend T* Cast(ClassName* InNode) \
+{ \
+	return Cast<T>((FSchematicGraphNode*) InNode); \
+} \
+template<typename T> \
+friend const T* CastChecked(const ClassName* InNode) \
+{ \
+	return CastChecked<T>((const FSchematicGraphNode*) InNode); \
+} \
+template<typename T> \
+friend T* CastChecked(ClassName* InNode) \
+{ \
+	return CastChecked<T>((FSchematicGraphNode*) InNode); \
+} \
+
+class ANIMATIONEDITORWIDGETS_API FSchematicGraphNode
+{
+public:
+	virtual ~FSchematicGraphNode() {}
+	
+	inline static const FName& Type = TEXT("FSchematicGraphNode");
+	virtual const FName& GetType() const;
+	virtual bool IsA(const FName& InType) const;
+	template<typename NodeType>
+	bool IsA() const
+	{
+		return IsA(NodeType::Type);
+	}
+	
+	template<typename T>
+	friend const T* Cast(const FSchematicGraphNode* InNode)
+	{
+		if(InNode)
+		{
+			if(InNode->IsA<T>())
+			{
+				return static_cast<const T*>(InNode);
+			}
+		}
+		return nullptr;
+	}
+
+	template<typename T>
+	friend T* Cast(FSchematicGraphNode* InNode)
+	{
+		if(InNode)
+		{
+			if(InNode->IsA<T>())
+			{
+				return static_cast<T*>(InNode);
+			}
+		}
+		return nullptr;
+	}
+
+	template<typename T>
+	friend const T* CastChecked(const FSchematicGraphNode* InNode)
+	{
+		const T* Node = Cast<T>(InNode);
+		check(Node);
+		return Node;
+	}
+
+	template<typename T>
+	friend T* CastChecked(FSchematicGraphNode* InNode)
+	{
+		T* Node = Cast<T>(InNode);
+		check(Node);
+		return Node;
+	}
+	
+	const FGuid& GetGuid() const { return Guid; }
+	bool IsSelected() const { return bIsSelected; }
+	void SetSelected(bool bSelected = true) { bIsSelected = bSelected; }
+	virtual FVector2d GetPosition() const { return Position; }
+	virtual FVector2d GetPositionOffset() const { return PositionOffset; }
+	virtual FLinearColor GetColor() const { return Color; }
+	virtual const FSlateBrush* GetBrush() const { return Brush; }
+	virtual const FText& GetToolTip() const { return ToolTip; }
+	virtual ESchematicGraphNodePlacementConstraint GetPlacement() const { return Placement; }
+	virtual void SetPlacement(ESchematicGraphNodePlacementConstraint InPlacement) { Placement = InPlacement; }
+	virtual ESchematicGraphNodeVisibility GetVisibility() const { return Visibility; }
+	virtual void SetVisibility(ESchematicGraphNodeVisibility InVisibility) { Visibility = InVisibility; }
+
+	virtual FString GetDragDropDecoratorLabel() const;
+
+protected:
+	
+	FGuid Guid = FGuid::NewGuid();
+	bool bIsSelected = false;
+	FVector2d Position = FVector2d::ZeroVector;
+	FVector2d PositionOffset = FVector2d::ZeroVector;
+	FLinearColor Color = FLinearColor::White;
+	const FSlateBrush* Brush = nullptr;
+	FText ToolTip = FText();
+	ESchematicGraphNodePlacementConstraint Placement = ESchematicGraphNodePlacementConstraint::Free;
+	ESchematicGraphNodeVisibility Visibility = ESchematicGraphNodeVisibility::Visible;
+
+	friend class FSchematicGraphModel;
+};
+
+DECLARE_EVENT_OneParam(FSchematicGraph, FOnNodeAdded, const FSchematicGraphNode*);
+DECLARE_EVENT_OneParam(FSchematicGraph, FOnNodeRemoved, const FSchematicGraphNode*);
+DECLARE_EVENT(FSchematicGraph, FOnGraphReset);
+
+class ANIMATIONEDITORWIDGETS_API FSchematicGraphModel
+{
+public:
+
+	virtual ~FSchematicGraphModel() {}
+	
+	virtual void Reset();
+	
+	template<typename NodeType = FSchematicGraphNode>
+	NodeType* AddNode(bool bNotify = true)
+	{
+		const TSharedPtr<FSchematicGraphNode> NewNode = MakeShareable(new NodeType);
+		Nodes.Add(NewNode);
+		NodeByGuid.Add(NewNode->GetGuid(), NewNode);
+
+		if (bNotify && OnNodeAddedDelegate.IsBound())
+		{
+			OnNodeAddedDelegate.Broadcast(NewNode.Get());
+		}
+		return static_cast<NodeType*>(NewNode.Get());
+	}
+
+	template<typename NodeType = FSchematicGraphNode>
+	const NodeType* FindNode(const FGuid& InGuid) const
+	{
+		if(const TSharedPtr<FSchematicGraphNode>* ExistingNode = NodeByGuid.Find(InGuid))
+		{
+			return Cast<NodeType>(ExistingNode->Get());
+		};
+		return nullptr;
+	}
+
+	template<typename NodeType = FSchematicGraphNode>
+	const NodeType* FindNodeChecked(const FGuid& InGuid) const
+	{
+		if(const TSharedPtr<FSchematicGraphNode>* ExistingNode = NodeByGuid.Find(InGuid))
+		{
+			check(ExistingNode->Get()->IsA(NodeType::Type));
+			return static_cast<NodeType*>(ExistingNode->Get());
+		};
+		return nullptr;
+	}
+
+	virtual bool RemoveNode(const FGuid& InGuid);
+
+	const TArray<TSharedPtr<FSchematicGraphNode>>& GetNodes() const { return Nodes; }
+
+	FVector2d GetPositionForNode(const FGuid& InGuid) const;
+	virtual FVector2d GetPositionForNode(const FSchematicGraphNode* InNode) const;
+	FVector2d GetPositionOffsetForNode(const FGuid& InGuid) const;
+	virtual FVector2d GetPositionOffsetForNode(const FSchematicGraphNode* InNode) const;
+	bool GetPositionAnimationEnabledForNode(const FGuid& InNodeGuid) const;
+	virtual bool GetPositionAnimationEnabledForNode(const FSchematicGraphNode* InNode) const;
+	FVector2d GetSizeForNode(const FGuid& InNodeGuid) const;
+	virtual FVector2d GetSizeForNode(const FSchematicGraphNode* InNode) const;
+	float GetScaleForNode(const FGuid& InNodeGuid) const;
+	virtual float GetScaleForNode(const FSchematicGraphNode* InNode) const;
+	FLinearColor GetColorForNode(const FGuid& InGuid) const;
+	virtual FLinearColor GetColorForNode(const FSchematicGraphNode* InNode) const;
+	const FSlateBrush* GetBrushForNode(const FGuid& InGuid) const;
+	virtual const FSlateBrush* GetBrushForNode(const FSchematicGraphNode* InNode) const;
+	const FText GetToolTipForNode(const FGuid& InGuid) const;
+	virtual const FText GetToolTipForNode(const FSchematicGraphNode* InNode) const;
+	ESchematicGraphNodePlacementConstraint GetPlacementForNode(const FGuid& InGuid) const;
+	virtual ESchematicGraphNodePlacementConstraint GetPlacementForNode(const FSchematicGraphNode* InNode) const;
+	ESchematicGraphNodeVisibility GetVisibilityForNode(const FGuid& InGuid) const;
+	virtual ESchematicGraphNodeVisibility GetVisibilityForNode(const FSchematicGraphNode* InNode) const;
+
+	FOnNodeAdded& OnNodeAdded() { return OnNodeAddedDelegate; }
+	FOnNodeRemoved& OnNodeRemoved() { return OnNodeRemovedDelegate; }
+	FOnGraphReset& OnGraphReset() { return OnGraphResetDelegate; }
+
+protected:
+
+	TArray<TSharedPtr<FSchematicGraphNode>> Nodes;
+	TMap<FGuid, TSharedPtr<FSchematicGraphNode>> NodeByGuid;
+	TMap<FGuid, TArray<FGuid>> Links;
+
+	FOnNodeAdded OnNodeAddedDelegate;
+	FOnNodeRemoved OnNodeRemovedDelegate;
+	FOnGraphReset OnGraphResetDelegate;
+};
+
+class SSchematicGraphNode;
+
+class FSchematicGraphNodeDragDropOp : public FDragDropOperation
+{
+public:
+	DRAG_DROP_OPERATOR_TYPE(FSchematicGraphNodeDragDropOp, FDragDropOperation)
+
+	DECLARE_DELEGATE_TwoParams(FOnEndDrag, SSchematicGraphNode*, const FDragDropOperation&);
+
+	static TSharedRef<FSchematicGraphNodeDragDropOp> New(TArray<SSchematicGraphNode*> InSchematicGraphNodes, const TArray<FGuid>& InElements, FOnEndDrag InOnEndDragDelegate);
+
+	virtual ~FSchematicGraphNodeDragDropOp() override;
+	
+	virtual TSharedPtr<SWidget> GetDefaultDecorator() const override;
+
+	/** @return true if this drag operation contains property paths */
+	bool HasElements() const
+	{
+		return Elements.Num() > 0;
+	}
+
+	/** @return The property paths from this drag operation */
+	const TArray<FGuid>& GetElements() const
+	{
+		return Elements;
+	}
+
+	FString GetJoinedDecoratorLabels() const;
+
+private:
+
+	/** Nodes being dragged */
+	TArray<SSchematicGraphNode*> SchematicGraphNodes;
+
+	/** Data for the property paths this item represents */
+	TArray<FGuid> Elements;
+
+	/** Delegate to call when this drag operation ends */
+	FOnEndDrag OnEndDragDelegate;
+};
+
+class ANIMATIONEDITORWIDGETS_API SSchematicGraphNode : public SNodePanel::SNode
+{
+public:
+
+	typedef TAnimatedAttribute<FVector2d> FVector2dAttribute;
+	typedef TAnimatedAttribute<float> FFloatAttribute;
+	typedef TAnimatedAttribute<FLinearColor> FLinearColorAttribute;
+	
+	DECLARE_DELEGATE_OneParam(FOnClicked, SSchematicGraphNode*);
+	DECLARE_DELEGATE_TwoParams(FOnBeginDrag, SSchematicGraphNode*, const FDragDropOperation&);
+	DECLARE_DELEGATE_TwoParams(FOnEndDrag, SSchematicGraphNode*, const FDragDropOperation&);
+	DECLARE_DELEGATE_TwoParams(FOnDrop, SSchematicGraphNode*, const FDragDropEvent&);
+
+	SLATE_BEGIN_ARGS(SSchematicGraphNode);
+	SLATE_ARGUMENT(const FSchematicGraphNode*, NodeData)
+	SLATE_ARGUMENT(TSharedPtr<FVector2dAttribute>, Position)
+	SLATE_ARGUMENT(TSharedPtr<FVector2dAttribute>, Size)
+	SLATE_ARGUMENT(TSharedPtr<FFloatAttribute>, Scale)
+	SLATE_ARGUMENT(TSharedPtr<FLinearColorAttribute>, Color)
+	SLATE_ATTRIBUTE(const FSlateBrush*, Brush)
+	SLATE_EVENT(FOnClicked, OnClicked)
+	SLATE_EVENT(FOnBeginDrag, OnBeginDrag)
+	SLATE_EVENT(FOnEndDrag, OnEndDrag)
+	SLATE_EVENT(FOnDrop, OnDrop)
+	SLATE_ATTRIBUTE(FText, ToolTipText)
+	SLATE_END_ARGS()
+
+	void Construct(const FArguments& InArgs);
+
+	virtual FVector2D ComputeDesiredSize(float LayoutScaleMultiplier) const override;
+	virtual int32 OnPaint(const FPaintArgs& Args, const FGeometry& AllottedGeometry, const FSlateRect& MyCullingRect, FSlateWindowElementList& OutDrawElements, int32 LayerId, const FWidgetStyle& InWidgetStyle, bool bParentEnabled) const override;
+	
+	virtual void OnMouseEnter(const FGeometry& MyGeometry, const FPointerEvent& MouseEvent) override;
+	virtual void OnMouseLeave(const FPointerEvent& MouseEvent) override;
+	virtual FReply OnMouseButtonDown(const FGeometry& MyGeometry, const FPointerEvent& MouseEvent) override;
+	virtual FReply OnDragOver(const FGeometry& MyGeometry, const FDragDropEvent& DragDropEvent) override;
+	virtual void OnDragLeave(const FDragDropEvent& DragDropEvent) override;
+	virtual FReply OnDrop(const FGeometry& MyGeometry, const FDragDropEvent& DragDropEvent) override;
+	virtual FReply OnDragDetected(const FGeometry& MyGeometry, const FPointerEvent& MouseEvent) override;
+
+	virtual FVector2d GetPosition() const override;
+	void EnablePositionAnimation(bool bEnabled = true);
+
+	const FVector2d& GetOriginalSize() const { return OriginalSize; }
+
+	const FSchematicGraphNode* GetNodeData() const { return NodeData; }
+	const FGuid GetGuid() const;
+	bool IsInteractive() const;
+
+	const bool IsBeingDragged() const { return bIsBeingDragged; } 
+	
+private:
+
+	bool bIsBeingDragged = false;
+	FVector2d OriginalSize = FVector2d(50.0,50.0);
+	float ScaledUp = 1.25;
+	float ScaledDown = 0.75;
+
+	const FSchematicGraphNode* NodeData = nullptr;
+	TSharedPtr<FVector2dAttribute> Position;
+	TOptional<FVector2d> PositionDuringDrag;
+	TOptional<FVector2d> OffsetDuringDrag;
+	TSharedPtr<FVector2dAttribute> Size;
+	TSharedPtr<FFloatAttribute> Scale;
+	TSharedPtr<FLinearColorAttribute> Color;
+	TAttribute<const FSlateBrush*> Brush;
+	FOnClicked OnClickedDelegate;
+	FOnBeginDrag OnBeginDragDelegate;
+	FOnEndDrag OnEndDragDelegate;
+	FOnDrop OnDropDelegate;
+
+	friend class SSchematicGraphPanel;
+};
+
+
+/** Widget allowing editing of a control rig's structure */
+class ANIMATIONEDITORWIDGETS_API SSchematicGraphPanel : public SNodePanel, public FTickableEditorObject
+{
+public:
+
+	using FVector2dAttribute = SSchematicGraphNode::FVector2dAttribute;
+	using FFloatAttribute = SSchematicGraphNode::FFloatAttribute;
+	using FLinearColorAttribute = SSchematicGraphNode::FLinearColorAttribute;
+
+	DECLARE_DELEGATE_TwoParams(FOnNodeClicked, SSchematicGraphPanel*, SSchematicGraphNode*);
+	DECLARE_DELEGATE_ThreeParams(FOnBeginDrag, SSchematicGraphPanel*, SSchematicGraphNode*, const FDragDropOperation&);
+	DECLARE_DELEGATE_ThreeParams(FOnEndDrag, SSchematicGraphPanel*, SSchematicGraphNode*, const FDragDropOperation&);
+	DECLARE_DELEGATE_ThreeParams(FOnDrop, SSchematicGraphPanel*, SSchematicGraphNode*, const FDragDropEvent&);
+	
+	SLATE_BEGIN_ARGS(SSchematicGraphPanel) {}
+	SLATE_ARGUMENT(bool, IsOverlay)
+	SLATE_ARGUMENT(FSchematicGraphModel*, GraphData)
+	SLATE_ARGUMENT(int32, PaddingLeft)
+	SLATE_ARGUMENT(int32, PaddingRight)
+	SLATE_ARGUMENT(int32, PaddingTop)
+	SLATE_ARGUMENT(int32, PaddingBottom)
+	SLATE_ARGUMENT(int32, PaddingInterNode)
+	SLATE_EVENT(FOnNodeClicked, OnNodeClicked)
+	SLATE_EVENT(FOnBeginDrag, OnBeginDrag)
+	SLATE_EVENT(FOnEndDrag, OnEndDrag)
+	SLATE_EVENT(FOnDrop, OnDrop)
+	SLATE_END_ARGS()
+
+	virtual ~SSchematicGraphPanel() override
+	{
+		OnNodeClickedDelegate.Unbind();
+		OnDropDelegate.Unbind();
+	}
+
+	void SetSchematicGraph(FSchematicGraphModel* InGraphData);
+	void Construct(const FArguments& InArgs);
+
+	void RebuildPanel();
+	void AddNode(const FSchematicGraphNode* InNodeToAdd);
+	void RemoveNode(const FSchematicGraphNode* InNodeToRemove);
+	const SSchematicGraphNode* FindNode(const FGuid& InGuid) const;
+	
+	// SNodePanel interface
+	virtual void OnArrangeChildren(const FGeometry& AllottedGeometry, FArrangedChildren& ArrangedChildren) const override;
+	virtual int32 OnPaint(const FPaintArgs& Args, const FGeometry& AllottedGeometry, const FSlateRect& MyCullingRect, FSlateWindowElementList& OutDrawElements, int32 LayerId, const FWidgetStyle& InWidgetStyle, bool bParentEnabled) const override;
+	virtual void RemoveAllNodes() override;
+	// End of SNodePanel interface
+	
+	TSharedRef<SSchematicGraphNode> GetChild(int32 ChildIndex) const;
+
+	// FTickableEditorObject Interface
+	virtual TStatId GetStatId() const override;
+	virtual void Tick(float DeltaTime) override;
+	virtual void Tick( const FGeometry& AllottedGeometry, const double InCurrentTime, const float InDeltaTime ) override;
+	virtual bool IsTickable() const override { return true; }
+	// End of FTickableEditorObject interface
+
+	void OnNodeClicked(SSchematicGraphNode* Node);
+	void OnBeginDragEvent(SSchematicGraphNode* Node, const FDragDropOperation& InDragDropEvent);
+	void OnEndDragEvent(SSchematicGraphNode* Node, const FDragDropOperation& InDragDropEvent);
+	void OnDropEvent(SSchematicGraphNode* Node, const FDragDropEvent& InDragDropEvent);
+
+	virtual FVector2d GetPositionForNode(FGuid InNodeGuid) const;
+	virtual FVector2d GetSizeForNode(FGuid InNodeGuid) const;
+	virtual float GetScaleForNode(FGuid InNodeGuid) const;
+	virtual FLinearColor GetColorForNode(FGuid InNodeGuid) const;
+	virtual const FSlateBrush* GetBrushForNode(FGuid InNodeGuid) const;
+	virtual FText GetToolTipForNode(FGuid InNodeGuid) const;
+
+	bool bIsDragDropping = false;
+	bool bIsOverlay;
+	int32 PaddingLeft = 0;
+	int32 PaddingRight = 0;
+	int32 PaddingTop = 0;
+	int32 PaddingBottom = 0;
+	int32 PaddingInterNode = 0;
+	TSharedPtr<FFloatAttribute> FadeBackgroundAlpha;
+	FSchematicGraphModel* GraphData;
+	FOnNodeClicked OnNodeClickedDelegate;
+	FOnBeginDrag OnBeginDragDelegate;
+	FOnEndDrag OnEndDragDelegate;
+	FOnDrop OnDropDelegate;
+	TArray<FGuid> NodesTopLeft;
+	TArray<FGuid> NodesTopRight;
+	TArray<FGuid> NodesBottomLeft;
+	TArray<FGuid> NodesBottomRight;
+	TMap<FGuid, TSharedPtr<SSchematicGraphNode>> NodeByGuid;
+};
+
+#endif
\ No newline at end of file
diff --git a/Engine/Source/Editor/AnimationEditorWidgets/Public/SchematicGraphPanel/SchematicViewportTabSummoner.h b/Engine/Source/Editor/AnimationEditorWidgets/Public/SchematicGraphPanel/SchematicViewportTabSummoner.h
new file mode 100644
--- /dev/null
+++ b/Engine/Source/Editor/AnimationEditorWidgets/Public/SchematicGraphPanel/SchematicViewportTabSummoner.h
@@ -0,0 +1,41 @@
+// Copyright Epic Games, Inc. All Rights Reserved.
+#if WITH_EDITOR
+
+#pragma once
+
+#include "CoreMinimal.h"
+#include "SSchematicGraphPanel.h"
+#include "WorkflowOrientedApp/WorkflowTabFactory.h"
+
+/** Called back when a viewport is created */
+DECLARE_DELEGATE_OneParam(FOnSchematicViewportCreated, const TSharedRef<class SSchematicGraphPanel>&);
+
+/** Arguments used to create a persona viewport tab */
+struct ANIMATIONEDITORWIDGETS_API FSchematicViewportArgs
+{
+	FSchematicViewportArgs()
+	{}
+
+	/** The model which contains the graph to display */
+	FSchematicGraphModel* SchematicGraph;
+	
+	/** Delegate fired when the viewport is created */
+	FOnSchematicViewportCreated OnViewportCreated;
+};
+
+
+struct ANIMATIONEDITORWIDGETS_API FSchematicViewportTabSummoner : public FWorkflowTabFactory
+{
+public:
+	static const FName TabID;
+	
+public:
+	FSchematicViewportTabSummoner(TSharedPtr<class FAssetEditorToolkit> InHostingApp, const FSchematicViewportArgs& InArgs);
+	
+	virtual FTabSpawnerEntry& RegisterTabSpawner(TSharedRef<FTabManager> TabManager, const FApplicationMode* CurrentApplicationMode) const;
+	virtual TSharedRef<SWidget> CreateTabBody(const FWorkflowTabSpawnInfo& Info) const override;
+
+	FSchematicGraphModel* SchematicGraph;
+	FOnSchematicViewportCreated OnViewportCreated;
+};
+#endif
\ No newline at end of file
