AnimationEditorWidgets - follow up to 30292594

#rnx

diff --git a/Engine/Source/Editor/AnimationEditorWidgets/Private/AnimationEditorWidgetsModule.cpp b/Engine/Source/Editor/AnimationEditorWidgets/Private/AnimationEditorWidgetsModule.cpp
--- a/Engine/Source/Editor/AnimationEditorWidgets/Private/AnimationEditorWidgetsModule.cpp
+++ b/Engine/Source/Editor/AnimationEditorWidgets/Private/AnimationEditorWidgetsModule.cpp
@@ -15,4 +15,4 @@ public:
 };
 
 
-IMPLEMENT_MODULE(FAnimationEditorWidgetsModule, AnimationWidgets);
+IMPLEMENT_MODULE(FAnimationEditorWidgetsModule, AnimationEditorWidgets);
