SchematicGraphPanel: Add support for auto scaling

#rnx

diff --git a/Engine/Plugins/Animation/ControlRig/Source/ControlRig/Private/Rigs/RigHierarchy.cpp b/Engine/Plugins/Animation/ControlRig/Source/ControlRig/Private/Rigs/RigHierarchy.cpp
--- a/Engine/Plugins/Animation/ControlRig/Source/ControlRig/Private/Rigs/RigHierarchy.cpp
+++ b/Engine/Plugins/Animation/ControlRig/Source/ControlRig/Private/Rigs/RigHierarchy.cpp
@@ -113,6 +113,7 @@ URigHierarchy::URigHierarchy()
 , HierarchyController(nullptr)
 , bIsControllerAvailable(true)
 , ResetPoseHash(INDEX_NONE)
+, bIsCopyingHierarchy(false)
 #if WITH_EDITOR
 , bPropagatingChange(false)
 , bForcePropagation(false)
@@ -480,6 +481,8 @@ void URigHierarchy::ResetToDefault()
 void URigHierarchy::CopyHierarchy(URigHierarchy* InHierarchy)
 {
 	check(InHierarchy);
+
+	const TGuardValue<bool> MarkCopyingHierarchy(bIsCopyingHierarchy, true);
 	
 	LLM_SCOPE_BYNAME(TEXT("Animation/ControlRig"));
 
@@ -654,6 +657,9 @@ void URigHierarchy::CopyHierarchy(URigHierarchy* InHierarchy)
 	}
 
 	EnsureCacheValidity();
+
+	bIsCopyingHierarchy = false;
+	Notify(ERigHierarchyNotification::HierarchyCopied, nullptr);
 }
 
 uint32 URigHierarchy::GetNameHash() const
diff --git a/Engine/Plugins/Animation/ControlRig/Source/ControlRig/Public/Rigs/RigHierarchy.h b/Engine/Plugins/Animation/ControlRig/Source/ControlRig/Public/Rigs/RigHierarchy.h
--- a/Engine/Plugins/Animation/ControlRig/Source/ControlRig/Public/Rigs/RigHierarchy.h
+++ b/Engine/Plugins/Animation/ControlRig/Source/ControlRig/Public/Rigs/RigHierarchy.h
@@ -206,6 +206,8 @@ public:
 	UFUNCTION(BlueprintCallable, Category = URigHierarchy)
 	void CopyHierarchy(URigHierarchy* InHierarchy);
 
+	bool IsCopyingHierarchy() const { return bIsCopyingHierarchy; }
+
 	/**
 	 * Returns true if the hierarchy currently has an execute context / the rig is running
 	 */
@@ -4302,6 +4304,8 @@ private:
 	mutable THierarchyCache<TMap<TTuple<int32, int32>, bool>> ElementDependencyCache;
 	mutable TArray<bool> ElementDependencyVisited;
 
+	bool bIsCopyingHierarchy;
+
 #if WITH_EDITOR
 
 	// this is mainly used for propagating changes between hierarchies in the direction of blueprint -> CDO -> other instances
diff --git a/Engine/Plugins/Animation/ControlRig/Source/ControlRig/Public/Rigs/RigHierarchyDefines.h b/Engine/Plugins/Animation/ControlRig/Source/ControlRig/Public/Rigs/RigHierarchyDefines.h
--- a/Engine/Plugins/Animation/ControlRig/Source/ControlRig/Public/Rigs/RigHierarchyDefines.h
+++ b/Engine/Plugins/Animation/ControlRig/Source/ControlRig/Public/Rigs/RigHierarchyDefines.h
@@ -94,6 +94,7 @@ enum class ERigHierarchyNotification : uint8
 	SocketColorChanged,
 	SocketDescriptionChanged,
 	SocketDesiredParentChanged,
+	HierarchyCopied,
 
 	/** MAX - invalid */
 	Max UMETA(Hidden),
diff --git a/Engine/Plugins/Animation/ControlRig/Source/ControlRigEditor/Private/ControlRigEditorStyle.h b/Engine/Plugins/Animation/ControlRig/Source/ControlRigEditor/Private/ControlRigEditorStyle.h
--- a/Engine/Plugins/Animation/ControlRig/Source/ControlRigEditor/Private/ControlRigEditorStyle.h
+++ b/Engine/Plugins/Animation/ControlRig/Source/ControlRigEditor/Private/ControlRigEditorStyle.h
@@ -36,6 +36,8 @@ class FControlRigEditorStyle final
 public:
 	FControlRigEditorStyle()
 		: FSlateStyleSet("ControlRigEditorStyle")
+		, BoneUserInterfaceColor(0.6f, 0.6f, 0.6f)
+		, NullUserInterfaceColor(0.75f, 0.75f, 0.75f)
 		, ConnectorUserInterfaceColor(0.0, 112.f/255.f, 224.f/255.f)
 		, SocketUserInterfaceColor(0.0, 112.f/255.f, 224.f/255.f)
 	{
@@ -171,6 +173,8 @@ public:
 	}
 
 	FSlateColor SpacePickerSelectColor;
+	FLinearColor BoneUserInterfaceColor;
+	FLinearColor NullUserInterfaceColor;
 	FLinearColor ConnectorUserInterfaceColor;
 	FLinearColor SocketUserInterfaceColor;
 };
diff --git a/Engine/Plugins/Animation/ControlRig/Source/ControlRigEditor/Private/ControlRigElementDetails.cpp b/Engine/Plugins/Animation/ControlRig/Source/ControlRigEditor/Private/ControlRigElementDetails.cpp
--- a/Engine/Plugins/Animation/ControlRig/Source/ControlRigEditor/Private/ControlRigElementDetails.cpp
+++ b/Engine/Plugins/Animation/ControlRig/Source/ControlRigEditor/Private/ControlRigElementDetails.cpp
@@ -5307,6 +5307,7 @@ void FRigSocketElementDetails::CustomizeSettings(IDetailLayoutBuilder& DetailBui
 
 	IDetailCategoryBuilder& SettingsCategory = DetailBuilder.EditCategory(TEXT("Settings"), LOCTEXT("Settings", "Settings"));
 
+	/*
 	SettingsCategory.AddCustomRow(FText::FromString(TEXT("Color")))
 	.NameContent()
 	[
@@ -5321,6 +5322,7 @@ void FRigSocketElementDetails::CustomizeSettings(IDetailLayoutBuilder& DetailBui
 		.Color(this, &FRigSocketElementDetails::GetSocketColor) 
 		.OnMouseButtonDown(this, &FRigSocketElementDetails::SetSocketColor)
 	];
+	*/
 
 	SettingsCategory.AddCustomRow(FText::FromString(TEXT("Description")))
 	.NameContent()
diff --git a/Engine/Plugins/Animation/ControlRig/Source/ControlRigEditor/Private/Editor/ControlRigSchematicModel.cpp b/Engine/Plugins/Animation/ControlRig/Source/ControlRigEditor/Private/Editor/ControlRigSchematicModel.cpp
--- a/Engine/Plugins/Animation/ControlRig/Source/ControlRigEditor/Private/Editor/ControlRigSchematicModel.cpp
+++ b/Engine/Plugins/Animation/ControlRig/Source/ControlRigEditor/Private/Editor/ControlRigSchematicModel.cpp
@@ -9,6 +9,7 @@
 #include "ModularRigRuleManager.h"
 #include "DragAndDrop/AssetDragDropOp.h"
 #include "Rigs/RigHierarchyController.h"
+#include "Async/TaskGraphInterfaces.h"
 
 FString FControlRigSchematicRigElementKeyNode::GetDragDropDecoratorLabel() const
 {
@@ -134,6 +135,11 @@ void FControlRigSchematicModel::OnSetObjectBeingDebugged(UObject* InObject)
 
 void FControlRigSchematicModel::OnHierarchyModified(ERigHierarchyNotification InNotif, URigHierarchy* InHierarchy, const FRigBaseElement* InElement)
 {
+	if(InHierarchy->IsCopyingHierarchy())
+	{
+		return;
+	}
+	
 	switch (InNotif)
 	{
 		case ERigHierarchyNotification::ElementAdded:
@@ -180,33 +186,42 @@ void FControlRigSchematicModel::OnHierarchyModified(ERigHierarchyNotification In
 			break;
 		}
 		case ERigHierarchyNotification::HierarchyReset:
+		case ERigHierarchyNotification::HierarchyCopied:
 		{
-			Reset();
-			TArray<FRigSocketElement*> Sockets = InHierarchy->GetElementsOfType<FRigSocketElement>();
+			const TArray<FRigSocketElement*> Sockets = InHierarchy->GetElementsOfType<FRigSocketElement>();
 			for (const FRigSocketElement* Socket : Sockets)
 			{
-				const FString NodeKey = Socket->GetKey().ToString();
-				if(!ContainsElementKeyNode(NodeKey))
+				if(!ContainsElementKeyNode(Socket->GetKey()))
+				{
+					AddElementKeyNode(Socket->GetKey());
+				}
+			}
+
+			const TArray<FRigConnectorElement*> Connectors = InHierarchy->GetElementsOfType<FRigConnectorElement>();
+			for (const FRigConnectorElement* Connector : Connectors)
+			{
+				if(!ContainsElementKeyNode(Connector->GetKey()))
 				{
-					AddElementKeyNode(NodeKey);
+					AddElementKeyNode(Connector->GetKey());
 				}
 			}
-			if (ControlRigBeingDebuggedPtr.IsValid())
+
+			// remove obsolete nodes
+			TArray<FGuid> GuidsToRemove;
+			for(const TSharedPtr<FSchematicGraphNode>& Node : Nodes)
 			{
-				const FRigElementKeyRedirector& Redirector = ControlRigBeingDebuggedPtr->GetElementKeyRedirector();
-				TArray<FRigConnectorElement*> Connectors = InHierarchy->GetElementsOfType<FRigConnectorElement>();
-				for (const FRigConnectorElement* Connector : Connectors)
+				if(const FControlRigSchematicRigElementKeyNode* ElementKeyNode = Cast<FControlRigSchematicRigElementKeyNode>(Node.Get()))
 				{
-					if (!Redirector.Contains(Connector->GetKey()))
+					if(!InHierarchy->Contains(ElementKeyNode->GetKey()))
 					{
-						const FString NodeKey = Connector->GetKey().ToString();
-						if(!ContainsElementKeyNode(NodeKey))
-						{
-							AddElementKeyNode(NodeKey);
-						}
+						GuidsToRemove.Add(ElementKeyNode->GetGuid());
 					}
 				}
 			}
+			for(const FGuid& Guid : GuidsToRemove)
+			{
+				RemoveNode(Guid);
+			}
 			break;
 		}
 		default:
@@ -222,26 +237,6 @@ void FControlRigSchematicModel::HandleModularRigModified(EModularRigNotification
 	{
 		case EModularRigNotification::ConnectionChanged:
 		{
-			if (InModule && ControlRigBeingDebuggedPtr.IsValid())
-			{
-				for (const TPair<FRigElementKey, FRigElementKey>& Pair : InModule->Connections)
-				{
-					RemoveElementKeyNode(Pair.Key);
-				}
-				
-				TArray<FRigConnectorElement*> Connectors = ControlRigBeingDebuggedPtr->GetHierarchy()->GetElementsOfType<FRigConnectorElement>();
-				for (const FRigConnectorElement* Connector : Connectors)
-				{
-					if (!InModule->Connections.Contains(Connector->GetKey()))
-					{
-						const FString NodeKey = Connector->GetKey().ToString();
-						if(!ContainsElementKeyNode(NodeKey))
-						{
-							AddElementKeyNode(NodeKey);
-						}
-					}
-				}
-			}
 			break;
 		}
 		default:
@@ -255,7 +250,17 @@ FVector2d FControlRigSchematicModel::GetPositionForNode(const FSchematicGraphNod
 {
 	if(const FControlRigSchematicRigElementKeyNode* Node = Cast<FControlRigSchematicRigElementKeyNode>(InNode))
 	{
-		if(Node->GetKey().Type != ERigElementType::Connector)
+		FRigElementKey Key = Node->GetKey();
+		if(Key.Type == ERigElementType::Connector)
+		{
+			FRigElementKey ResolvedKey;
+			if(IsConnectorResolved(Key, &ResolvedKey))
+			{
+				Key = ResolvedKey;
+			}
+		}
+
+		if(Key.IsValid())
 		{
 			if(const TSharedPtr<FControlRigEditor> Editor = ControlRigEditor.Pin())
 			{
@@ -265,7 +270,7 @@ FVector2d FControlRigSchematicModel::GetPositionForNode(const FSchematicGraphNod
 					{
 						if(const URigHierarchy* Hierarchy = ControlRig->GetHierarchy())
 						{
-							const FTransform Transform = Hierarchy->GetGlobalTransform(Node->GetKey());
+							const FTransform Transform = Hierarchy->GetGlobalTransform(Key);
 							return Editor->ComputePersonaProjectedScreenPos(Transform.GetLocation());
 						}
 					}
@@ -357,6 +362,46 @@ const FSlateBrush* FControlRigSchematicModel::GetBrushForNode(const FSchematicGr
 	return FSchematicGraphModel::GetBrushForNode(InNode);
 }
 
+FLinearColor FControlRigSchematicModel::GetColorForNode(const FSchematicGraphNode* InNode) const
+{
+	if(const FControlRigSchematicRigElementKeyNode* ElementKeyNode = Cast<FControlRigSchematicRigElementKeyNode>(InNode))
+	{
+		switch(ElementKeyNode->GetKey().Type)
+		{
+			case ERigElementType::Bone:
+			{
+				return FControlRigEditorStyle::Get().BoneUserInterfaceColor;
+			}
+			case ERigElementType::Null:
+			{
+				return FControlRigEditorStyle::Get().NullUserInterfaceColor;
+			}
+			case ERigElementType::Control:
+			{
+				if (ControlRigBeingDebuggedPtr.IsValid())
+				{
+					if (const URigHierarchy* Hierarchy = ControlRigBeingDebuggedPtr->GetHierarchy())
+					{
+						if(const FRigControlElement* Control = Hierarchy->Find<FRigControlElement>(ElementKeyNode->GetKey()))
+						{
+							return Control->Settings.ShapeColor;
+						}
+					}
+				}
+			}
+			case ERigElementType::Socket:
+			{
+				return FControlRigEditorStyle::Get().SocketUserInterfaceColor;
+			}
+			case ERigElementType::Connector:
+			{
+				return FControlRigEditorStyle::Get().ConnectorUserInterfaceColor;
+			}
+		}
+	}
+	return FSchematicGraphModel::GetColorForNode(InNode);
+}
+
 const FText FControlRigSchematicModel::GetToolTipForNode(const FSchematicGraphNode* InNode) const
 {
 	if(const FControlRigSchematicRigElementKeyNode* Node = Cast<FControlRigSchematicRigElementKeyNode>(InNode))
@@ -397,7 +442,10 @@ ESchematicGraphNodePlacementConstraint FControlRigSchematicModel::GetPlacementFo
 	{
 		if(Node->GetKey().Type == ERigElementType::Connector)
 		{
-			return ESchematicGraphNodePlacementConstraint::BottomRight;
+			if(!IsConnectorResolved(Node->GetKey()))
+			{
+				return ESchematicGraphNodePlacementConstraint::BottomRight;
+			}
 		}
 	}
 	return FSchematicGraphModel::GetPlacementForNode(InNode);
@@ -534,7 +582,7 @@ void FControlRigSchematicModel::HandleSchematicDrop(SSchematicGraphPanel* InPane
 	}
 
 	const FControlRigSchematicRigElementKeyNode* ElementKeyNode =
-	Cast<FControlRigSchematicRigElementKeyNode>(InNode->GetNodeData());
+		Cast<FControlRigSchematicRigElementKeyNode>(InNode->GetNodeData());
 	if(ElementKeyNode == nullptr)
 	{
 		return;
@@ -551,8 +599,9 @@ void FControlRigSchematicModel::HandleSchematicDrop(SSchematicGraphPanel* InPane
 	{
 		return;
 	}
-	
-	FRigBaseElement* Target = Hierarchy->Find(ElementKeyNode->GetKey());
+
+	const FRigElementKey& TargetKey = ElementKeyNode->GetKey();
+	FRigBaseElement* Target = Hierarchy->Find(TargetKey);
 	if (!Target)
 	{
 		return;
@@ -572,55 +621,101 @@ void FControlRigSchematicModel::HandleSchematicDrop(SSchematicGraphPanel* InPane
 
 			if(UControlRigBlueprint* AssetBlueprint = Cast<UControlRigBlueprint>(AssetData.GetAsset()))
 			{
-				if (UModularRigController* Controller = ControlRigBlueprint->GetModularRigController())
+				FFunctionGraphTask::CreateAndDispatchWhenReady([this, AssetBlueprint, TargetKey]()
 				{
-					const FName ModuleName = Controller->GetSafeNewName(FString(), FRigName(AssetBlueprint->RigModuleSettings.Identifier.Name));
-					const FString ModulePath = Controller->AddModule(ModuleName, AssetBlueprint->GetControlRigClass(), FString());
-					if(!ModulePath.IsEmpty())
+					if (UModularRigController* Controller = ControlRigBlueprint->GetModularRigController())
 					{
-						FRigElementKey PrimaryConnectorKey;
-						TArray<FRigConnectorElement*> Connectors = Hierarchy->GetElementsOfType<FRigConnectorElement>();
-						for (FRigConnectorElement* Connector : Connectors)
+						UControlRig* ControlRig = ControlRigBlueprint->GetDebuggedControlRig();
+						if (!ControlRig)
+						{
+							return;
+						}
+
+						URigHierarchy* Hierarchy = ControlRig->GetHierarchy();
+						if (!Hierarchy)
+						{
+							return;
+						}
+
+						const FName ModuleName = Controller->GetSafeNewName(FString(), FRigName(AssetBlueprint->RigModuleSettings.Identifier.Name));
+						const FString ModulePath = Controller->AddModule(ModuleName, AssetBlueprint->GetControlRigClass(), FString());
+						if(!ModulePath.IsEmpty())
 						{
-							if (Connector->Settings.Type == EConnectorType::Primary)
+							FRigElementKey PrimaryConnectorKey;
+							TArray<FRigConnectorElement*> Connectors = Hierarchy->GetElementsOfType<FRigConnectorElement>();
+							for (FRigConnectorElement* Connector : Connectors)
 							{
-								FString Path, Name;
-								Connector->GetName().Split(UModularRig::NamespaceSeparator, &Path, &Name, ESearchCase::CaseSensitive, ESearchDir::FromEnd);
-								if (Path == ModulePath)
+								if (Connector->Settings.Type == EConnectorType::Primary)
 								{
-									PrimaryConnectorKey = Connector->GetKey();
-									break;
+									FString Path, Name;
+									Connector->GetName().Split(UModularRig::NamespaceSeparator, &Path, &Name, ESearchCase::CaseSensitive, ESearchDir::FromEnd);
+									if (Path == ModulePath)
+									{
+										PrimaryConnectorKey = Connector->GetKey();
+										break;
+									}
 								}
 							}
-						}
 
-						const FName TargetModulePath = Hierarchy->GetNameMetadata(Target->GetKey(), URigHierarchy::NameSpaceMetadataName, NAME_None);
-						if(!TargetModulePath.IsNone())
-						{
-							(void)Controller->ReparentModule(ModulePath, TargetModulePath.ToString());
+							const FName TargetModulePath = Hierarchy->GetNameMetadata(TargetKey, URigHierarchy::NameSpaceMetadataName, NAME_None);
+							if(!TargetModulePath.IsNone())
+							{
+								(void)Controller->ReparentModule(ModulePath, TargetModulePath.ToString());
+							}
+							Controller->ConnectConnectorToElement(PrimaryConnectorKey, TargetKey);
 						}
-						Controller->ConnectConnectorToElement(PrimaryConnectorKey, Target->GetKey());
 					}
-				}
+				}, TStatId(), NULL, ENamedThreads::GameThread);
 			}
 		}
 	}
 	else if(SchematicDragDropOp.IsValid())
 	{
-		if (UModularRigController* Controller = ControlRigBlueprint->GetModularRigController())
+		const TArray<FGuid> Sources = SchematicDragDropOp->GetElements();
+		FFunctionGraphTask::CreateAndDispatchWhenReady([this, Sources, TargetKey]()
 		{
-			const TArray<FGuid> Sources = SchematicDragDropOp->GetElements();
-			for (const TPair<FRigElementKey, FGuid> Pair : RigElementKeyToGuid)
+			UControlRig* ControlRig = ControlRigBlueprint->GetDebuggedControlRig();
+			if (!ControlRig)
+			{
+				return;
+			}
+
+			URigHierarchy* Hierarchy = ControlRig->GetHierarchy();
+			if (!Hierarchy)
+			{
+				return;
+			}
+
+			if (UModularRigController* Controller = ControlRigBlueprint->GetModularRigController())
 			{
-				if(Sources.Contains(Pair.Value))
+				for (const TPair<FRigElementKey, FGuid> Pair : RigElementKeyToGuid)
 				{
-					if (Hierarchy->Find<FRigConnectorElement>(Pair.Key))
+					if(Sources.Contains(Pair.Value))
 					{
-						Controller->ConnectConnectorToElement(Pair.Key, Target->GetKey());
-						break;
+						if (Hierarchy->Find<FRigConnectorElement>(Pair.Key))
+						{
+							Controller->ConnectConnectorToElement(Pair.Key, TargetKey);
+							break;
+						}
 					}
 				}
 			}
+		}, TStatId(), NULL, ENamedThreads::GameThread);
+	}
+}
+
+bool FControlRigSchematicModel::IsConnectorResolved(const FRigElementKey& InConnectorKey, FRigElementKey* OutKey) const
+{
+	if(ControlRigBlueprint.IsValid())
+	{
+		if(const FRigElementKey* TargetKey = ControlRigBlueprint->ConnectionMap.Find(InConnectorKey))
+		{
+			if(OutKey)
+			{
+				*OutKey = *TargetKey;
+			}
+			return true;
 		}
 	}
-}
\ No newline at end of file
+	return false;
+}
diff --git a/Engine/Plugins/Animation/ControlRig/Source/ControlRigEditor/Private/Editor/ControlRigSchematicModel.h b/Engine/Plugins/Animation/ControlRig/Source/ControlRigEditor/Private/Editor/ControlRigSchematicModel.h
--- a/Engine/Plugins/Animation/ControlRig/Source/ControlRigEditor/Private/Editor/ControlRigSchematicModel.h
+++ b/Engine/Plugins/Animation/ControlRig/Source/ControlRigEditor/Private/Editor/ControlRigSchematicModel.h
@@ -23,6 +23,7 @@ public:
 	
 	const FRigElementKey& GetKey() const { return Key; }
 	virtual FString GetDragDropDecoratorLabel() const override;
+	virtual bool IsAutoScaleEnabled() const override { return true; }
 	
 protected:
 
@@ -53,6 +54,7 @@ public:
 	virtual FVector2d GetPositionForNode(const FSchematicGraphNode* InNode) const override;
 	virtual bool GetPositionAnimationEnabledForNode(const FSchematicGraphNode* InNode) const override;
 	virtual const FSlateBrush* GetBrushForNode(const FSchematicGraphNode* InNode) const override;
+	virtual FLinearColor GetColorForNode(const FSchematicGraphNode* InNode) const override;
 	virtual const FText GetToolTipForNode(const FSchematicGraphNode* InNode) const override;
 	virtual ESchematicGraphNodePlacementConstraint GetPlacementForNode(const FSchematicGraphNode* InNode) const override;
 
@@ -62,6 +64,7 @@ private:
 	void HandleSchematicBeginDrag(SSchematicGraphPanel* InPanel, SSchematicGraphNode* InNode, const FDragDropOperation& InDragDropOperation);
 	void HandleSchematicEndDrag(SSchematicGraphPanel* InPanel, SSchematicGraphNode* InNode, const FDragDropOperation& InDragDropOperation);
 	void HandleSchematicDrop(SSchematicGraphPanel* InPanel, SSchematicGraphNode* InNode, const FDragDropEvent& InDragDropEvent);
+	bool IsConnectorResolved(const FRigElementKey& InConnectorKey, FRigElementKey* OutKey = nullptr) const;
 	
 	TWeakPtr<FControlRigEditor> ControlRigEditor;
 	TWeakObjectPtr<UControlRigBlueprint> ControlRigBlueprint;
diff --git a/Engine/Plugins/Animation/ControlRig/Source/ControlRigEditor/Private/Editor/SRigHierarchy.cpp b/Engine/Plugins/Animation/ControlRig/Source/ControlRigEditor/Private/Editor/SRigHierarchy.cpp
--- a/Engine/Plugins/Animation/ControlRig/Source/ControlRigEditor/Private/Editor/SRigHierarchy.cpp
+++ b/Engine/Plugins/Animation/ControlRig/Source/ControlRigEditor/Private/Editor/SRigHierarchy.cpp
@@ -2011,127 +2011,161 @@ void SRigHierarchy::HandleNewItem(ERigElementType InElementType, bool bIsAnimati
 		// unselect current selected item
 		ClearDetailPanel();
 
+		const bool bAllowMultipleItems =
+			InElementType == ERigElementType::Socket ||
+			InElementType == ERigElementType::Null;
+
 		URigHierarchyController* Controller = Hierarchy->GetController(true);
 		check(Controller);
 
 		FScopedTransaction Transaction(LOCTEXT("HierarchyTreeAdded", "Add new item to hierarchy"));
 
-		FRigElementKey ParentKey;
-		FTransform ParentTransform = FTransform::Identity;
-
 		TArray<FRigElementKey> SelectedKeys = GetSelectedKeys();
-		if (SelectedKeys.Num() > 0)
+		if (SelectedKeys.Num() > 1 && !bAllowMultipleItems)
 		{
-			ParentKey = SelectedKeys[0];
-			ParentTransform = Hierarchy->GetGlobalTransform(ParentKey);
+			SelectedKeys = {SelectedKeys[0]};
 		}
-
-		// use bone's name as prefix if creating a control
-		FString NewNameTemplate;
-		const bool bIsParentABone = ParentKey.IsValid() && ParentKey.Type == ERigElementType::Bone;
-		if( InElementType == ERigElementType::Control && bIsParentABone )
+		else if(SelectedKeys.IsEmpty())
 		{
-			static const FString CtrlSuffix(TEXT("_ctrl"));
-			NewNameTemplate = ParentKey.Name.ToString();
-			NewNameTemplate += CtrlSuffix;
+			SelectedKeys = {FRigElementKey()};
 		}
-		else
+
+		for(const FRigElementKey& SelectedKey : SelectedKeys)
 		{
-			NewNameTemplate = FString::Printf(TEXT("New%s"), *StaticEnum<ERigElementType>()->GetNameStringByValue((int64)InElementType));
+			FRigElementKey ParentKey;
+			FTransform ParentTransform = FTransform::Identity;
 
-			if(bIsAnimationChannel)
+			if(SelectedKey.IsValid())
 			{
-				static const FString NewAnimationChannel = TEXT("Channel");
-				NewNameTemplate = NewAnimationChannel;
+				ParentKey = SelectedKey;
+				ParentTransform = Hierarchy->GetGlobalTransform(ParentKey);
 			}
-		}
-		
-		const FName NewElementName = CreateUniqueName(*NewNameTemplate, InElementType);
-		{
-			TGuardValue<bool> GuardRigHierarchyChanges(bIsChangingRigHierarchy, true);
-			switch (InElementType)
+
+			// use bone's name as prefix if creating a control
+			FString NewNameTemplate;
+			if(ParentKey.IsValid() && ParentKey.Type == ERigElementType::Bone)
 			{
-				case ERigElementType::Bone:
+				NewNameTemplate = ParentKey.Name.ToString();
+
+				if(InElementType == ERigElementType::Control)
 				{
-					NewItemKey = Controller->AddBone(NewElementName, ParentKey, ParentTransform, true, ERigBoneType::User, true, true);
-					break;
+					static const FString CtrlSuffix(TEXT("_ctrl"));
+					NewNameTemplate += CtrlSuffix;
+				}
+				else if(InElementType == ERigElementType::Null)
+				{
+					static const FString NullSuffix(TEXT("_null"));
+					NewNameTemplate += NullSuffix;
 				}
-				case ERigElementType::Control:
+				else if(InElementType == ERigElementType::Socket)
 				{
-					FRigControlSettings Settings;
+					static const FString SocketSuffix(TEXT("_socket"));
+					NewNameTemplate += SocketSuffix;
+				}
+				else
+				{
+					NewNameTemplate.Reset();
+				}
+			}
 
-					if(bIsAnimationChannel)
-					{
-						Settings.AnimationType = ERigControlAnimationType::AnimationChannel;
-						Settings.ControlType = ERigControlType::Float;
-						Settings.MinimumValue = FRigControlValue::Make<float>(0.f);
-						Settings.MaximumValue = FRigControlValue::Make<float>(1.f);
-						Settings.DisplayName = Hierarchy->GetSafeNewDisplayName(ParentKey, FRigName(NewNameTemplate));
+			if(NewNameTemplate.IsEmpty())
+			{
+				NewNameTemplate = FString::Printf(TEXT("New%s"), *StaticEnum<ERigElementType>()->GetNameStringByValue((int64)InElementType));
 
-						NewItemKey = Controller->AddAnimationChannel(NewElementName, ParentKey, Settings, true, true);
+				if(bIsAnimationChannel)
+				{
+					static const FString NewAnimationChannel = TEXT("Channel");
+					NewNameTemplate = NewAnimationChannel;
+				}
+			}
+			
+			const FName NewElementName = CreateUniqueName(*NewNameTemplate, InElementType);
+			{
+				TGuardValue<bool> GuardRigHierarchyChanges(bIsChangingRigHierarchy, true);
+				switch (InElementType)
+				{
+					case ERigElementType::Bone:
+					{
+						NewItemKey = Controller->AddBone(NewElementName, ParentKey, ParentTransform, true, ERigBoneType::User, true, true);
+						break;
 					}
-					else
+					case ERigElementType::Control:
 					{
-						Settings.ControlType = ERigControlType::EulerTransform;
-						FEulerTransform Identity = FEulerTransform::Identity;
-						FRigControlValue ValueToSet = FRigControlValue::Make<FEulerTransform>(Identity);
-						Settings.MinimumValue = ValueToSet;
-						Settings.MaximumValue = ValueToSet;
+						FRigControlSettings Settings;
 
-						NewItemKey = Controller->AddControl(NewElementName, ParentKey, Settings, Settings.GetIdentityValue(), FTransform::Identity, FTransform::Identity, true, true);
+						if(bIsAnimationChannel)
+						{
+							Settings.AnimationType = ERigControlAnimationType::AnimationChannel;
+							Settings.ControlType = ERigControlType::Float;
+							Settings.MinimumValue = FRigControlValue::Make<float>(0.f);
+							Settings.MaximumValue = FRigControlValue::Make<float>(1.f);
+							Settings.DisplayName = Hierarchy->GetSafeNewDisplayName(ParentKey, FRigName(NewNameTemplate));
+
+							NewItemKey = Controller->AddAnimationChannel(NewElementName, ParentKey, Settings, true, true);
+						}
+						else
+						{
+							Settings.ControlType = ERigControlType::EulerTransform;
+							FEulerTransform Identity = FEulerTransform::Identity;
+							FRigControlValue ValueToSet = FRigControlValue::Make<FEulerTransform>(Identity);
+							Settings.MinimumValue = ValueToSet;
+							Settings.MaximumValue = ValueToSet;
+
+							NewItemKey = Controller->AddControl(NewElementName, ParentKey, Settings, Settings.GetIdentityValue(), FTransform::Identity, FTransform::Identity, true, true);
+						}
+						break;
 					}
-					break;
-				}
-				case ERigElementType::Null:
-				{
-					NewItemKey = Controller->AddNull(NewElementName, ParentKey, ParentTransform, true, true, true);
-					break;
-				}
-				case ERigElementType::Connector:
-				{
-					FString FailureReason;
-					if(!ControlRigBlueprint->CanTurnIntoControlRigModule(false, &FailureReason))
+					case ERigElementType::Null:
 					{
-						if(ControlRigBlueprint->Hierarchy->Num(ERigElementType::Connector) == 0)
+						NewItemKey = Controller->AddNull(NewElementName, ParentKey, ParentTransform, true, true, true);
+						break;
+					}
+					case ERigElementType::Connector:
+					{
+						FString FailureReason;
+						if(!ControlRigBlueprint->CanTurnIntoControlRigModule(false, &FailureReason))
 						{
-							static constexpr TCHAR Format[] = TEXT("Connector cannot be created: %s");
-							UE_LOG(LogControlRig, Warning, Format, *FailureReason);
-							FNotificationInfo Info(FText::FromString(FString::Printf(Format, *FailureReason)));
-							Info.bUseSuccessFailIcons = true;
-							Info.Image = FAppStyle::GetBrush(TEXT("MessageLog.Warning"));
-							Info.bFireAndForget = true;
-							Info.bUseThrobber = true;
-							Info.FadeOutDuration = 2.f;
-							Info.ExpireDuration = 8.f;;
-							TSharedPtr<SNotificationItem> NotificationPtr = FSlateNotificationManager::Get().AddNotification(Info);
-							if (NotificationPtr)
+							if(ControlRigBlueprint->Hierarchy->Num(ERigElementType::Connector) == 0)
 							{
-								NotificationPtr->SetCompletionState(SNotificationItem::CS_Fail);
+								static constexpr TCHAR Format[] = TEXT("Connector cannot be created: %s");
+								UE_LOG(LogControlRig, Warning, Format, *FailureReason);
+								FNotificationInfo Info(FText::FromString(FString::Printf(Format, *FailureReason)));
+								Info.bUseSuccessFailIcons = true;
+								Info.Image = FAppStyle::GetBrush(TEXT("MessageLog.Warning"));
+								Info.bFireAndForget = true;
+								Info.bUseThrobber = true;
+								Info.FadeOutDuration = 2.f;
+								Info.ExpireDuration = 8.f;;
+								TSharedPtr<SNotificationItem> NotificationPtr = FSlateNotificationManager::Get().AddNotification(Info);
+								if (NotificationPtr)
+								{
+									NotificationPtr->SetCompletionState(SNotificationItem::CS_Fail);
+								}
+								return;
 							}
-							return;
 						}
-					}
 
-					const bool bIsPrimary = Hierarchy->GetConnectorKeys(false).Num() == 0;
-					FRigConnectorSettings Settings;
-					Settings.Type = bIsPrimary ? EConnectorType::Primary : EConnectorType::Secondary;
-						if(!bIsPrimary)
-						{
-							Settings.Rules.Reset();
-							Settings.AddRule(FRigChildOfPrimaryConnectionRule());
-						}
-					NewItemKey = Controller->AddConnector(NewElementName, Settings, true);
-					(void)ResolveConnector(NewItemKey, ParentKey);
-					break;
-				}
-				case ERigElementType::Socket:
-				{
-					NewItemKey = Controller->AddSocket(NewElementName, ParentKey, ParentTransform, true, FRigSocketElement::SocketDefaultColor, FString(), true, true);
-					break;
-				}
-				default:
-				{
-					return;
+						const bool bIsPrimary = Hierarchy->GetConnectorKeys(false).Num() == 0;
+						FRigConnectorSettings Settings;
+						Settings.Type = bIsPrimary ? EConnectorType::Primary : EConnectorType::Secondary;
+							if(!bIsPrimary)
+							{
+								Settings.Rules.Reset();
+								Settings.AddRule(FRigChildOfPrimaryConnectionRule());
+							}
+						NewItemKey = Controller->AddConnector(NewElementName, Settings, true);
+						(void)ResolveConnector(NewItemKey, ParentKey);
+						break;
+					}
+					case ERigElementType::Socket:
+					{
+						NewItemKey = Controller->AddSocket(NewElementName, ParentKey, ParentTransform, true, FRigSocketElement::SocketDefaultColor, FString(), true, true);
+						break;
+					}
+					default:
+					{
+						return;
+					}
 				}
 			}
 		}
diff --git a/Engine/Source/Editor/AnimationEditorWidgets/Private/SchematicGraphPanel/SSchematicGraphPanel.cpp b/Engine/Source/Editor/AnimationEditorWidgets/Private/SchematicGraphPanel/SSchematicGraphPanel.cpp
--- a/Engine/Source/Editor/AnimationEditorWidgets/Private/SchematicGraphPanel/SSchematicGraphPanel.cpp
+++ b/Engine/Source/Editor/AnimationEditorWidgets/Private/SchematicGraphPanel/SSchematicGraphPanel.cpp
@@ -111,20 +111,54 @@ FVector2d FSchematicGraphModel::GetSizeForNode(const FSchematicGraphNode* InNode
 	return DefaultSize;
 }
 
-float FSchematicGraphModel::GetScaleForNode(const FGuid& InNodeGuid) const
+float FSchematicGraphModel::GetScaleForNode(const FGuid& InNodeGuid, bool bIncludeScaleOffset) const
 {
 	if(const FSchematicGraphNode* Node = FindNode(InNodeGuid))
 	{
-		return GetScaleForNode(Node);
+		return GetScaleForNode(Node, bIncludeScaleOffset);
 	}
 	return 1.f;
 }
 
-float FSchematicGraphModel::GetScaleForNode(const FSchematicGraphNode* InNode) const
+float FSchematicGraphModel::GetScaleForNode(const FSchematicGraphNode* InNode, bool bIncludeScaleOffset) const
 {
+	if(bIncludeScaleOffset)
+	{
+		return GetScaleOffsetForNode(InNode);
+	}
+	return 1.f;
+}
+
+float FSchematicGraphModel::GetScaleOffsetForNode(const FGuid& InNodeGuid) const
+{
+	if(const FSchematicGraphNode* Node = FindNode(InNodeGuid))
+	{
+		return GetScaleOffsetForNode(Node);
+	}
 	return 1.f;
 }
 
+float FSchematicGraphModel::GetScaleOffsetForNode(const FSchematicGraphNode* InNode) const
+{
+	check(InNode);
+	return InNode->GetScaleOffset();
+}
+
+bool FSchematicGraphModel::IsAutoScaleEnabledForNode(const FGuid& InNodeGuid) const
+{
+	if(const FSchematicGraphNode* Node = FindNode(InNodeGuid))
+	{
+		return IsAutoScaleEnabledForNode(Node);
+	}
+	return false;
+}
+
+bool FSchematicGraphModel::IsAutoScaleEnabledForNode(const FSchematicGraphNode* InNode) const
+{
+	check(InNode);
+	return InNode->IsAutoScaleEnabled();
+}
+
 FLinearColor FSchematicGraphModel::GetColorForNode(const FGuid& InGuid) const
 {
 	if(const FSchematicGraphNode* Node = FindNode(InGuid))
@@ -257,6 +291,7 @@ FString FSchematicGraphNodeDragDropOp::GetJoinedDecoratorLabels() const
 
 SSchematicGraphNode::FArguments::FArguments()
 : _NodeData(nullptr)
+, _EnableAutoScale(false)
 , _Brush(nullptr)
 {
 	static const FSchematicGraphNode EmptyNodeData;
@@ -269,7 +304,7 @@ void SSchematicGraphNode::Construct(const FArguments& InArgs)
 {
 	if(InArgs._NodeData)
 	{
-		NodeData = InArgs._NodeData;
+		NodeData = const_cast<FSchematicGraphNode*>(InArgs._NodeData);
 	}
 	OnClickedDelegate = InArgs._OnClicked;
 	OnBeginDragDelegate = InArgs._OnBeginDrag;
@@ -279,6 +314,7 @@ void SSchematicGraphNode::Construct(const FArguments& InArgs)
 	Position = InArgs._Position;
 	Size = InArgs._Size;
 	Scale = InArgs._Scale;
+	EnableAutoScale = InArgs._EnableAutoScale;
 	Color = InArgs._Color;
 	if(InArgs._Brush.IsSet() || InArgs._Brush.IsBound())
 	{
@@ -290,6 +326,8 @@ void SSchematicGraphNode::Construct(const FArguments& InArgs)
 	{
 		SetToolTipText(InArgs._ToolTipText);
 	}
+
+	SetVisibility(TAttribute<EVisibility>::CreateSP(this, &SSchematicGraphNode::GetNodeVisibility));
 }
 
 FVector2D SSchematicGraphNode::ComputeDesiredSize(float LayoutScaleMultiplier) const
@@ -324,7 +362,11 @@ void SSchematicGraphNode::OnMouseEnter(const FGeometry& MyGeometry, const FPoint
 		return;
 	}
 	SNode::OnMouseEnter(MyGeometry, MouseEvent);
-	Scale->Set(ScaledUp);
+
+	if(NodeData)
+	{
+		NodeData->SetScaleOffset(ScaledUp);
+	}
 }
 
 void SSchematicGraphNode::OnMouseLeave(const FPointerEvent& MouseEvent)
@@ -334,7 +376,11 @@ void SSchematicGraphNode::OnMouseLeave(const FPointerEvent& MouseEvent)
 		return;
 	}
 	SNode::OnMouseLeave(MouseEvent);
-	Scale->Set(1.f);
+
+	if(NodeData)
+	{
+		NodeData->SetScaleOffset(1.f);
+	}
 }
 
 FReply SSchematicGraphNode::OnDragOver(const FGeometry& MyGeometry, const FDragDropEvent& DragDropEvent)
@@ -344,8 +390,13 @@ FReply SSchematicGraphNode::OnDragOver(const FGeometry& MyGeometry, const FDragD
 		return FReply::Unhandled();
 	}
 	SNode::OnDragOver(MyGeometry, DragDropEvent);
-	Scale->Set(ScaledDown);
+
+	if(NodeData)
+	{
+		NodeData->SetScaleOffset(ScaledDown);
+	}
 	return FReply::Handled();
+	
 }
 
 void SSchematicGraphNode::OnDragLeave(const FDragDropEvent& DragDropEvent)
@@ -355,7 +406,11 @@ void SSchematicGraphNode::OnDragLeave(const FDragDropEvent& DragDropEvent)
 		return;
 	}
 	SNode::OnDragLeave(DragDropEvent);
-	Scale->Set(1.f);
+
+	if(NodeData)
+	{
+		NodeData->SetScaleOffset(1.f);
+	}
 }
 
 FReply SSchematicGraphNode::OnDrop(const FGeometry& MyGeometry, const FDragDropEvent& DragDropEvent)
@@ -406,6 +461,29 @@ FReply SSchematicGraphNode::OnDragDetected(const FGeometry& MyGeometry, const FP
 	return FReply::Unhandled();
 }
 
+EVisibility SSchematicGraphNode::GetNodeVisibility() const
+{
+	if(SchematicGraphPanel)
+	{
+		ESchematicGraphNodeVisibility Vis = SchematicGraphPanel->GetVisibilityForNode(GetGuid());
+		if(Vis == ESchematicGraphNodeVisibility::Hidden)
+		{
+			return EVisibility::Hidden;
+		}
+		if(Vis == ESchematicGraphNodeVisibility::FadedOut)
+		{
+			return EVisibility::HitTestInvisible;
+		}
+	}
+
+	if(IsBeingDragged())
+	{
+		return EVisibility::HitTestInvisible;
+	}
+
+	return EVisibility::Visible;
+}
+
 FReply SSchematicGraphNode::OnMouseButtonDown(const FGeometry& MyGeometry, const FPointerEvent& MouseEvent)
 {
 	SNode::OnMouseButtonDown(MyGeometry, MouseEvent);
@@ -511,7 +589,7 @@ void SSchematicGraphPanel::AddNode(const FSchematicGraphNode* InNodeToAdd)
 	const FGuid Guid = InNodeToAdd->GetGuid();
 	const auto Position = FVector2dAttribute::CreateWithGetter(Vector2DInterpSettings, FVector2dAttribute::FGetter::CreateSP(this, &SSchematicGraphPanel::GetPositionForNode, Guid));
 	const auto Size = FVector2dAttribute::CreateWithGetter(Vector2DInterpSettings, FVector2dAttribute::FGetter::CreateSP(this, &SSchematicGraphPanel::GetSizeForNode, Guid));
-	const auto Scale = FFloatAttribute::CreateWithGetter(FloatInterpSettings, FFloatAttribute::FGetter::CreateSP(this, &SSchematicGraphPanel::GetScaleForNode, Guid), 0.f);
+	const auto Scale = FFloatAttribute::CreateWithGetter(FloatInterpSettings, FFloatAttribute::FGetter::CreateSP(this, &SSchematicGraphPanel::GetScaleForNode, Guid, true), 0.f);
 	const auto Color = FLinearColorAttribute::CreateWithGetter(ColorInterpSettings, FLinearColorAttribute::FGetter::CreateSP(this, &SSchematicGraphPanel::GetColorForNode, Guid));
 
 	const TSharedRef<SSchematicGraphNode> NewNode = SNew(SSchematicGraphNode)
@@ -519,6 +597,7 @@ void SSchematicGraphPanel::AddNode(const FSchematicGraphNode* InNodeToAdd)
 														.Size(Size)
 														.Scale(Scale)
 														.Color(Color)
+														.EnableAutoScale(this, &SSchematicGraphPanel::IsAutoScaleEnabledForNode, InNodeToAdd->GetGuid())
 														.Brush(this, &SSchematicGraphPanel::GetBrushForNode, InNodeToAdd->GetGuid())
 														.ToolTipText(this, &SSchematicGraphPanel::GetToolTipForNode, InNodeToAdd->GetGuid())
 														.OnClicked_Raw(this, &SSchematicGraphPanel::OnNodeClicked)
@@ -527,6 +606,7 @@ void SSchematicGraphPanel::AddNode(const FSchematicGraphNode* InNodeToAdd)
 														.OnDrop_Raw(this, &SSchematicGraphPanel::OnDropEvent)
 														.NodeData(InNodeToAdd);
 	SNodePanel::AddGraphNode(NewNode);
+	NewNode->SchematicGraphPanel = this;
 	NodeByGuid.Add(Guid, NewNode.ToSharedPtr());
 }
 
@@ -645,6 +725,12 @@ void SSchematicGraphPanel::RemoveAllNodes()
 	SNodePanel::RemoveAllNodes();
 }
 
+FReply SSchematicGraphPanel::OnMouseWheel(const FGeometry& MyGeometry, const FPointerEvent& MouseEvent)
+{
+	// disable mouse wheel for now
+	return FReply::Unhandled();
+}
+
 TSharedRef<SSchematicGraphNode> SSchematicGraphPanel::GetChild(int32 ChildIndex) const
 {
 	return StaticCastSharedRef<SSchematicGraphNode>(Children[ChildIndex]);
@@ -673,7 +759,8 @@ void SSchematicGraphPanel::Tick(float DeltaTime)
 	NodesBottomLeft.Reset();
 	NodesBottomRight.Reset();
 
-	FVector2D Size = GetCachedGeometry().Size;
+	const FVector2D Size = GetCachedGeometry().Size;
+
 	for (int32 i=0; i<Children.Num(); ++i)
 	{
 		TSharedRef<SSchematicGraphNode> Node = GetChild(i);
@@ -729,6 +816,8 @@ void SSchematicGraphPanel::Tick(float DeltaTime)
 			}
 		}
 	}
+
+	UpdateAutoScalingForNodes();
 }
 
 void SSchematicGraphPanel::Tick(const FGeometry& AllottedGeometry, const double InCurrentTime, const float InDeltaTime)
@@ -785,29 +874,37 @@ FVector2d SSchematicGraphPanel::GetPositionForNode(FGuid InNodeGuid) const
 				case ESchematicGraphNodePlacementConstraint::TopLeft:
 				{
 					const int32 NodeIndexInCorner = NodesTopLeft.IndexOfByKey(InNodeGuid);
-					check(NodeIndexInCorner != INDEX_NONE);
-					Position += FVector2d(PaddingLeft + HalfOriginalSize.X, PaddingTop + HalfOriginalSize.Y + ((OriginalSize.Y + PaddingInterNode) * NodeIndexInCorner));
+					if(NodeIndexInCorner != INDEX_NONE)
+					{
+						Position += FVector2d(PaddingLeft + HalfOriginalSize.X, PaddingTop + HalfOriginalSize.Y + ((OriginalSize.Y + PaddingInterNode) * NodeIndexInCorner));
+					}
 					break;
 				}
 				case ESchematicGraphNodePlacementConstraint::TopRight:
 				{
 					const int32 NodeIndexInCorner = NodesTopRight.IndexOfByKey(InNodeGuid);
-					check(NodeIndexInCorner != INDEX_NONE);
-					Position += FVector2d(Size.X - PaddingRight - HalfOriginalSize.X, PaddingTop + HalfOriginalSize.Y + ((OriginalSize.Y + PaddingInterNode) * NodeIndexInCorner));
+					if(NodeIndexInCorner != INDEX_NONE)
+					{
+						Position += FVector2d(Size.X - PaddingRight - HalfOriginalSize.X, PaddingTop + HalfOriginalSize.Y + ((OriginalSize.Y + PaddingInterNode) * NodeIndexInCorner));
+					}
 					break;
 				}
 				case ESchematicGraphNodePlacementConstraint::BottomLeft:
 				{
 					const int32 NodeIndexInCorner = NodesBottomLeft.IndexOfByKey(InNodeGuid);
-					check(NodeIndexInCorner != INDEX_NONE);
-					Position += FVector2d(PaddingLeft + HalfOriginalSize.X, Size.Y - PaddingBottom - HalfOriginalSize.Y - ((OriginalSize.Y + PaddingInterNode)  * NodeIndexInCorner));
+					if(NodeIndexInCorner != INDEX_NONE)
+					{
+						Position += FVector2d(PaddingLeft + HalfOriginalSize.X, Size.Y - PaddingBottom - HalfOriginalSize.Y - ((OriginalSize.Y + PaddingInterNode)  * NodeIndexInCorner));
+					}
 					break;
 				}
 				case ESchematicGraphNodePlacementConstraint::BottomRight:
 				{
 					const int32 NodeIndexInCorner = NodesBottomRight.IndexOfByKey(InNodeGuid);
-					check(NodeIndexInCorner != INDEX_NONE);
-					Position += FVector2d(Size.X - PaddingRight - HalfOriginalSize.X, Size.Y - PaddingBottom - HalfOriginalSize.Y - ((OriginalSize.Y + PaddingInterNode)  * NodeIndexInCorner));
+					if(NodeIndexInCorner != INDEX_NONE)
+					{
+						Position += FVector2d(Size.X - PaddingRight - HalfOriginalSize.X, Size.Y - PaddingBottom - HalfOriginalSize.Y - ((OriginalSize.Y + PaddingInterNode)  * NodeIndexInCorner));
+					}
 					break;
 				}
 			}
@@ -827,15 +924,38 @@ FVector2d SSchematicGraphPanel::GetSizeForNode(FGuid InNodeGuid) const
 	return DefaultSize;
 }
 
-float SSchematicGraphPanel::GetScaleForNode(FGuid InNodeGuid) const
+float SSchematicGraphPanel::GetScaleForNode(FGuid InNodeGuid, bool bIncludeScaleOffset) const
 {
+	// check if the node may be auto scaled
+	if(const TSharedPtr<SSchematicGraphNode>* NodePtr = NodeByGuid.Find(InNodeGuid))
+	{
+		if(NodePtr->Get()->AutoScale.IsSet())
+		{
+			float ScaleOffset = 1.f;
+			if(GraphData)
+			{
+				ScaleOffset = GraphData->GetScaleOffsetForNode(NodePtr->Get()->GetNodeData());
+			}
+			return NodePtr->Get()->AutoScale.GetValue() * ScaleOffset;
+		}
+	}
+	
 	if(GraphData)
 	{
-		return GraphData->GetScaleForNode(InNodeGuid);
+		return GraphData->GetScaleForNode(InNodeGuid, bIncludeScaleOffset);
 	}
 	return 1.f;
 }
 
+bool SSchematicGraphPanel::IsAutoScaleEnabledForNode(FGuid InNodeGuid) const
+{
+	if(GraphData)
+	{
+		return GraphData->IsAutoScaleEnabledForNode(InNodeGuid);
+	}
+	return false;
+}
+
 FLinearColor SSchematicGraphPanel::GetColorForNode(FGuid InNodeGuid) const
 {
 	if(GraphData)
@@ -864,6 +984,100 @@ FText SSchematicGraphPanel::GetToolTipForNode(FGuid InNodeGuid) const
 	return FText();
 }
 
+ESchematicGraphNodeVisibility SSchematicGraphPanel::GetVisibilityForNode(FGuid InNodeGuid) const
+{
+	if(GraphData)
+	{
+		return GraphData->GetVisibilityForNode(InNodeGuid);
+	}
+	return ESchematicGraphNodeVisibility::Visible;
+}
+
+void SSchematicGraphPanel::UpdateAutoScalingForNodes()
+{
+	TArray<bool> NodeIsAutoScaling;
+	TArray<FVector2d> NodePositions;
+	TArray<double> NodeRadiuses;
+	NodeIsAutoScaling.Reserve(Children.Num());
+	NodePositions.Reserve(Children.Num());
+	NodeRadiuses.Reserve(Children.Num());
+
+	for (int32 i=0; i<Children.Num(); ++i)
+	{
+		TSharedRef<SSchematicGraphNode> Node = GetChild(i);
+
+		NodeIsAutoScaling.Add(!Node->bIsBeingDragged && Node->EnableAutoScale.Get());
+		if(NodeIsAutoScaling.Last())
+		{
+			NodePositions.Add(GetPositionForNode(Node->GetGuid()));
+			const FVector2d NodeSize = GetSizeForNode(Node->GetGuid());
+			// note: this is not necessarily the best way to determine the radius of a node
+			NodeRadiuses.Add(FMath::Min(NodeSize.X, NodeSize.Y) * 0.5);
+		}
+		else
+		{
+			NodePositions.Add(FVector2d::ZeroVector);
+			NodeRadiuses.Add(0);
+		}
+	}
+
+	// for now brute force find all neighbors
+	// and determine how much of the radius we have
+	// to reduce to avoid overlap.
+	// todo: use a faster distance algorithm
+	TArray<double> RadiusReductionPerNode;
+	RadiusReductionPerNode.AddZeroed(Children.Num());
+	
+	for (int32 i=0; i<Children.Num(); ++i)
+	{
+		if(!NodeIsAutoScaling[i])
+		{
+			continue;
+		}
+		
+		const FVector2d& PositionA = NodePositions[i];
+		const double RadiusA = NodeRadiuses[i];
+
+		for (int32 j=i+1; j<Children.Num(); ++j)
+		{
+			if(!NodeIsAutoScaling[j])
+			{
+				continue;
+			}
+
+			const FVector2d& PositionB = NodePositions[j];
+			const double RadiusB = NodeRadiuses[j];
+			static constexpr double AutoScalePadding = 4.0;
+			const double MinDistance = RadiusA + RadiusB + AutoScalePadding;
+
+			const double Distance = (PositionA - PositionB).Size();
+			if(Distance < SMALL_NUMBER || Distance > MinDistance)
+			{
+				continue;
+			}
+
+			const double RadiusReduction = (MinDistance - Distance) * 0.5;
+			RadiusReductionPerNode[i] = FMath::Max(RadiusReductionPerNode[i], RadiusReduction);
+			RadiusReductionPerNode[j] = FMath::Max(RadiusReductionPerNode[j], RadiusReduction);
+		}
+	}
+
+	// mark nodes for auto scaling
+	for (int32 i=0; i<Children.Num(); ++i)
+	{
+		TSharedRef<SSchematicGraphNode> Node = GetChild(i);
+		if(RadiusReductionPerNode[i] > SMALL_NUMBER)
+		{
+			const float Scale = (NodeRadiuses[i] - RadiusReductionPerNode[i]) / NodeRadiuses[i];
+			Node->AutoScale = Scale;
+		}
+		else
+		{
+			Node->AutoScale.Reset();
+		}
+	}
+}
+
 #undef LOCTEXT_NAMESPACE
 
 #endif
\ No newline at end of file
diff --git a/Engine/Source/Editor/AnimationEditorWidgets/Public/SchematicGraphPanel/SSchematicGraphPanel.h b/Engine/Source/Editor/AnimationEditorWidgets/Public/SchematicGraphPanel/SSchematicGraphPanel.h
--- a/Engine/Source/Editor/AnimationEditorWidgets/Public/SchematicGraphPanel/SSchematicGraphPanel.h
+++ b/Engine/Source/Editor/AnimationEditorWidgets/Public/SchematicGraphPanel/SSchematicGraphPanel.h
@@ -8,6 +8,8 @@
 #include "TickableEditorObject.h"
 #include "Animation/AnimatedAttribute.h"
 
+class SSchematicGraphPanel;
+
 enum ESchematicGraphNodePlacementConstraint
 {
 	Free,
@@ -117,8 +119,12 @@ public:
 	void SetSelected(bool bSelected = true) { bIsSelected = bSelected; }
 	virtual FVector2d GetPosition() const { return Position; }
 	virtual FVector2d GetPositionOffset() const { return PositionOffset; }
+	virtual void SetPositionOffset(const FVector2d& InPositionOffset) { PositionOffset = InPositionOffset; }
+	virtual float GetScaleOffset() const { return ScaleOffset; }
+	virtual void SetScaleOffset(float InScaleOffset) { ScaleOffset = InScaleOffset; }
 	virtual FLinearColor GetColor() const { return Color; }
 	virtual const FSlateBrush* GetBrush() const { return Brush; }
+	virtual bool IsAutoScaleEnabled() const { return false; }
 	virtual const FText& GetToolTip() const { return ToolTip; }
 	virtual ESchematicGraphNodePlacementConstraint GetPlacement() const { return Placement; }
 	virtual void SetPlacement(ESchematicGraphNodePlacementConstraint InPlacement) { Placement = InPlacement; }
@@ -133,6 +139,7 @@ protected:
 	bool bIsSelected = false;
 	FVector2d Position = FVector2d::ZeroVector;
 	FVector2d PositionOffset = FVector2d::ZeroVector;
+	float ScaleOffset = 1.f;
 	FLinearColor Color = FLinearColor::White;
 	const FSlateBrush* Brush = nullptr;
 	FText ToolTip = FText();
@@ -201,8 +208,12 @@ public:
 	virtual bool GetPositionAnimationEnabledForNode(const FSchematicGraphNode* InNode) const;
 	FVector2d GetSizeForNode(const FGuid& InNodeGuid) const;
 	virtual FVector2d GetSizeForNode(const FSchematicGraphNode* InNode) const;
-	float GetScaleForNode(const FGuid& InNodeGuid) const;
-	virtual float GetScaleForNode(const FSchematicGraphNode* InNode) const;
+	float GetScaleForNode(const FGuid& InNodeGuid, bool bIncludeScaleOffset) const;
+	virtual float GetScaleForNode(const FSchematicGraphNode* InNode, bool bIncludeScaleOffset) const;
+	float GetScaleOffsetForNode(const FGuid& InNodeGuid) const;
+	virtual float GetScaleOffsetForNode(const FSchematicGraphNode* InNode) const;
+	bool IsAutoScaleEnabledForNode(const FGuid& InNodeGuid) const;
+	virtual bool IsAutoScaleEnabledForNode(const FSchematicGraphNode* InNode) const;
 	FLinearColor GetColorForNode(const FGuid& InGuid) const;
 	virtual FLinearColor GetColorForNode(const FSchematicGraphNode* InNode) const;
 	const FSlateBrush* GetBrushForNode(const FGuid& InGuid) const;
@@ -288,6 +299,7 @@ public:
 	SLATE_ARGUMENT(TSharedPtr<FVector2dAttribute>, Position)
 	SLATE_ARGUMENT(TSharedPtr<FVector2dAttribute>, Size)
 	SLATE_ARGUMENT(TSharedPtr<FFloatAttribute>, Scale)
+	SLATE_ATTRIBUTE(bool, EnableAutoScale)
 	SLATE_ARGUMENT(TSharedPtr<FLinearColorAttribute>, Color)
 	SLATE_ATTRIBUTE(const FSlateBrush*, Brush)
 	SLATE_EVENT(FOnClicked, OnClicked)
@@ -310,6 +322,7 @@ public:
 	virtual FReply OnDrop(const FGeometry& MyGeometry, const FDragDropEvent& DragDropEvent) override;
 	virtual FReply OnDragDetected(const FGeometry& MyGeometry, const FPointerEvent& MouseEvent) override;
 
+	virtual EVisibility GetNodeVisibility() const;
 	virtual FVector2d GetPosition() const override;
 	void EnablePositionAnimation(bool bEnabled = true);
 
@@ -325,21 +338,24 @@ private:
 
 	bool bIsBeingDragged = false;
 	FVector2d OriginalSize = FVector2d(50.0,50.0);
-	float ScaledUp = 1.25;
-	float ScaledDown = 0.75;
+	static inline constexpr float ScaledUp = 1.25;
+	static inline constexpr float ScaledDown = 0.75;
 
-	const FSchematicGraphNode* NodeData = nullptr;
+	FSchematicGraphNode* NodeData = nullptr;
 	TSharedPtr<FVector2dAttribute> Position;
 	TOptional<FVector2d> PositionDuringDrag;
 	TOptional<FVector2d> OffsetDuringDrag;
 	TSharedPtr<FVector2dAttribute> Size;
 	TSharedPtr<FFloatAttribute> Scale;
+	TAttribute<bool> EnableAutoScale;
+	TOptional<float> AutoScale;
 	TSharedPtr<FLinearColorAttribute> Color;
 	TAttribute<const FSlateBrush*> Brush;
 	FOnClicked OnClickedDelegate;
 	FOnBeginDrag OnBeginDragDelegate;
 	FOnEndDrag OnEndDragDelegate;
 	FOnDrop OnDropDelegate;
+	SSchematicGraphPanel* SchematicGraphPanel = nullptr;
 
 	friend class SSchematicGraphPanel;
 };
@@ -391,6 +407,7 @@ public:
 	virtual void OnArrangeChildren(const FGeometry& AllottedGeometry, FArrangedChildren& ArrangedChildren) const override;
 	virtual int32 OnPaint(const FPaintArgs& Args, const FGeometry& AllottedGeometry, const FSlateRect& MyCullingRect, FSlateWindowElementList& OutDrawElements, int32 LayerId, const FWidgetStyle& InWidgetStyle, bool bParentEnabled) const override;
 	virtual void RemoveAllNodes() override;
+	virtual FReply OnMouseWheel(const FGeometry& MyGeometry, const FPointerEvent& MouseEvent) override;
 	// End of SNodePanel interface
 	
 	TSharedRef<SSchematicGraphNode> GetChild(int32 ChildIndex) const;
@@ -409,10 +426,14 @@ public:
 
 	virtual FVector2d GetPositionForNode(FGuid InNodeGuid) const;
 	virtual FVector2d GetSizeForNode(FGuid InNodeGuid) const;
-	virtual float GetScaleForNode(FGuid InNodeGuid) const;
+	virtual float GetScaleForNode(FGuid InNodeGuid, bool bIncludeScaleOffset) const;
+	virtual bool IsAutoScaleEnabledForNode(FGuid InNodeGuid) const;
 	virtual FLinearColor GetColorForNode(FGuid InNodeGuid) const;
 	virtual const FSlateBrush* GetBrushForNode(FGuid InNodeGuid) const;
 	virtual FText GetToolTipForNode(FGuid InNodeGuid) const;
+	virtual ESchematicGraphNodeVisibility GetVisibilityForNode(FGuid InNodeGuid) const;
+
+	void UpdateAutoScalingForNodes();
 
 	bool bIsDragDropping = false;
 	bool bIsOverlay;
