SchematicPanel: Allow nodes to control drag support

#rnx

diff --git a/Engine/Plugins/Animation/ControlRig/Source/ControlRigEditor/Private/Editor/ControlRigSchematicModel.cpp b/Engine/Plugins/Animation/ControlRig/Source/ControlRigEditor/Private/Editor/ControlRigSchematicModel.cpp
--- a/Engine/Plugins/Animation/ControlRig/Source/ControlRigEditor/Private/Editor/ControlRigSchematicModel.cpp
+++ b/Engine/Plugins/Animation/ControlRig/Source/ControlRigEditor/Private/Editor/ControlRigSchematicModel.cpp
@@ -16,6 +16,11 @@ FString FControlRigSchematicRigElementKeyNode::GetDragDropDecoratorLabel() const
 	return Key.ToString();
 }
 
+bool FControlRigSchematicRigElementKeyNode::IsDragSupported() const
+{
+	return Key.Type == ERigElementType::Connector;
+}
+
 FControlRigSchematicModel::~FControlRigSchematicModel()
 {
 	if(ControlRigBeingDebuggedPtr.IsValid())
@@ -436,6 +441,25 @@ const FText FControlRigSchematicModel::GetToolTipForNode(const FSchematicGraphNo
 	return FSchematicGraphModel::GetToolTipForNode(InNode);
 }
 
+ESchematicGraphNodeVisibility FControlRigSchematicModel::GetVisibilityForNode(const FSchematicGraphNode* InNode) const
+{
+	if(const FControlRigSchematicRigElementKeyNode* Node = Cast<FControlRigSchematicRigElementKeyNode>(InNode))
+	{
+		if(Node->GetKey().Type == ERigElementType::Connector)
+		{
+			FRigElementKey ResolvedSocket;
+			if(IsConnectorResolved(Node->GetKey(), &ResolvedSocket))
+			{
+				if(ContainsElementKeyNode(ResolvedSocket))
+				{
+					return ESchematicGraphNodeVisibility::Hidden;
+				}
+			}
+		}
+	}
+	return FSchematicGraphModel::GetVisibilityForNode(InNode);
+}
+
 ESchematicGraphNodePlacementConstraint FControlRigSchematicModel::GetPlacementForNode(const FSchematicGraphNode* InNode) const
 {
 	if(const FControlRigSchematicRigElementKeyNode* Node = Cast<FControlRigSchematicRigElementKeyNode>(InNode))
@@ -451,6 +475,25 @@ ESchematicGraphNodePlacementConstraint FControlRigSchematicModel::GetPlacementFo
 	return FSchematicGraphModel::GetPlacementForNode(InNode);
 }
 
+bool FControlRigSchematicModel::GetForwardedNodeForDrag(FGuid& InOutGuid) const
+{
+	if(const FControlRigSchematicRigElementKeyNode* ElementKeyNode = Cast<FControlRigSchematicRigElementKeyNode>(FindNode(InOutGuid)))
+	{
+		for(const TPair<FRigElementKey, FRigElementKey>& Pair : ControlRigBlueprint->ConnectionMap)
+		{
+			if(Pair.Value == ElementKeyNode->GetKey())
+			{
+				if(const FControlRigSchematicRigElementKeyNode* ForwardedNode = FindElementKeyNode(Pair.Key))
+				{
+					InOutGuid = ForwardedNode->GetGuid();
+					return true;
+				}
+			}
+		}
+	}
+	return FSchematicGraphModel::GetForwardedNodeForDrag(InOutGuid);
+}
+
 void FControlRigSchematicModel::HandleSchematicNodeClicked(SSchematicGraphPanel* InPanel, SSchematicGraphNode* InNode)
 {
 	for (const TSharedPtr<FSchematicGraphNode>& Node : Nodes)
@@ -532,6 +575,11 @@ void FControlRigSchematicModel::HandleSchematicBeginDrag(SSchematicGraphPanel* I
 		return;
 	}
 
+	if (UModularRigController* Controller = ControlRigBlueprint->GetModularRigController())
+	{
+		Controller->DisconnectConnector(Connector->GetKey(), true);
+	}
+
 	const UModularRigRuleManager* RuleManager = Hierarchy->GetRuleManager();
 	const FModularRigResolveResult Result = RuleManager->FindMatches(Connector, ModuleInstance, ControlRigBeingDebuggedPtr->GetElementKeyRedirector());
 
diff --git a/Engine/Plugins/Animation/ControlRig/Source/ControlRigEditor/Private/Editor/ControlRigSchematicModel.h b/Engine/Plugins/Animation/ControlRig/Source/ControlRigEditor/Private/Editor/ControlRigSchematicModel.h
--- a/Engine/Plugins/Animation/ControlRig/Source/ControlRigEditor/Private/Editor/ControlRigSchematicModel.h
+++ b/Engine/Plugins/Animation/ControlRig/Source/ControlRigEditor/Private/Editor/ControlRigSchematicModel.h
@@ -24,6 +24,7 @@ public:
 	const FRigElementKey& GetKey() const { return Key; }
 	virtual FString GetDragDropDecoratorLabel() const override;
 	virtual bool IsAutoScaleEnabled() const override { return true; }
+	virtual bool IsDragSupported() const override;
 	
 protected:
 
@@ -56,8 +57,11 @@ public:
 	virtual const FSlateBrush* GetBrushForNode(const FSchematicGraphNode* InNode) const override;
 	virtual FLinearColor GetColorForNode(const FSchematicGraphNode* InNode) const override;
 	virtual const FText GetToolTipForNode(const FSchematicGraphNode* InNode) const override;
+	virtual ESchematicGraphNodeVisibility GetVisibilityForNode(const FSchematicGraphNode* InNode) const override;
 	virtual ESchematicGraphNodePlacementConstraint GetPlacementForNode(const FSchematicGraphNode* InNode) const override;
 
+	virtual bool GetForwardedNodeForDrag(FGuid& InOutGuid) const override;
+
 private:
 
 	void HandleSchematicNodeClicked(SSchematicGraphPanel* InPanel, SSchematicGraphNode* InNode);
diff --git a/Engine/Source/Editor/AnimationEditorWidgets/Private/SchematicGraphPanel/SSchematicGraphPanel.cpp b/Engine/Source/Editor/AnimationEditorWidgets/Private/SchematicGraphPanel/SSchematicGraphPanel.cpp
--- a/Engine/Source/Editor/AnimationEditorWidgets/Private/SchematicGraphPanel/SSchematicGraphPanel.cpp
+++ b/Engine/Source/Editor/AnimationEditorWidgets/Private/SchematicGraphPanel/SSchematicGraphPanel.cpp
@@ -244,6 +244,21 @@ ESchematicGraphNodeVisibility FSchematicGraphModel::GetVisibilityForNode(const F
 	return InNode->GetVisibility();
 }
 
+bool FSchematicGraphModel::IsDragSupportedForNode(const FGuid& InGuid) const
+{
+	if(const FSchematicGraphNode* Node = FindNode(InGuid))
+	{
+		return IsDragSupportedForNode(Node);
+	}
+	return false;
+}
+
+bool FSchematicGraphModel::IsDragSupportedForNode(const FSchematicGraphNode* InNode) const
+{
+	check(InNode);
+	return InNode->IsDragSupported();
+}
+
 TSharedRef<FSchematicGraphNodeDragDropOp> FSchematicGraphNodeDragDropOp::New(TArray<SSchematicGraphNode*> InSchematicGraphNodes, const TArray<FGuid>& InElements, FSchematicGraphNodeDragDropOp::FOnEndDrag InOnEndDragDelegate)
 {
 	TSharedRef<FSchematicGraphNodeDragDropOp> Operation = MakeShared<FSchematicGraphNodeDragDropOp>();
@@ -444,18 +459,34 @@ FReply SSchematicGraphNode::OnDrop(const FGeometry& MyGeometry, const FDragDropE
 
 FReply SSchematicGraphNode::OnDragDetected(const FGeometry& MyGeometry, const FPointerEvent& MouseEvent)
 {
-	TArray<FGuid> DraggedElements = {GetGuid()};
+	FGuid Guid = GetGuid();
+	if(SchematicGraphPanel)
+	{
+		const FReply ReplyFromPanel = SchematicGraphPanel->HandleNodeDragDetected(Guid, MyGeometry, MouseEvent);
+		if(ReplyFromPanel.IsEventHandled())
+		{
+			return ReplyFromPanel;
+		}
+	}
+	
+	TArray<FGuid> DraggedElements = {Guid};
 	if (MouseEvent.IsMouseButtonDown(EKeys::LeftMouseButton) && DraggedElements.Num() > 0)
 	{
-		bIsBeingDragged = true;
+		if(SchematicGraphPanel)
+		{
+			if(SchematicGraphPanel->IsDragSupportedForNode(GetGuid()))
+			{
+				bIsBeingDragged = true;
 
-		const FVector2f AbsoluteMousePosition = MouseEvent.GetScreenSpacePosition();
-		const FVector2d LocalMousePosition =  (AbsoluteMousePosition - MyGeometry.GetAbsolutePosition()) / MyGeometry.GetAccumulatedLayoutTransform().GetScale();
-		OffsetDuringDrag = -LocalMousePosition;
-		
-		const TSharedRef<FSchematicGraphNodeDragDropOp> DragDropOp = FSchematicGraphNodeDragDropOp::New({this}, MoveTemp(DraggedElements), OnEndDragDelegate);
-		OnBeginDragDelegate.ExecuteIfBound(this, DragDropOp.Get());
-		return FReply::Handled().BeginDragDrop(DragDropOp);
+				const FVector2f AbsoluteMousePosition = MouseEvent.GetScreenSpacePosition();
+				const FVector2d LocalMousePosition =  (AbsoluteMousePosition - MyGeometry.GetAbsolutePosition()) / MyGeometry.GetAccumulatedLayoutTransform().GetScale();
+				OffsetDuringDrag = -LocalMousePosition;
+			
+				const TSharedRef<FSchematicGraphNodeDragDropOp> DragDropOp = FSchematicGraphNodeDragDropOp::New({this}, MoveTemp(DraggedElements), OnEndDragDelegate);
+				OnBeginDragDelegate.ExecuteIfBound(this, DragDropOp.Get());
+				return FReply::Handled().BeginDragDrop(DragDropOp);
+			}
+		}
 	}
 	
 	return FReply::Unhandled();
@@ -463,6 +494,11 @@ FReply SSchematicGraphNode::OnDragDetected(const FGeometry& MyGeometry, const FP
 
 EVisibility SSchematicGraphNode::GetNodeVisibility() const
 {
+	if(IsBeingDragged())
+	{
+		return EVisibility::HitTestInvisible;
+	}
+
 	if(SchematicGraphPanel)
 	{
 		ESchematicGraphNodeVisibility Vis = SchematicGraphPanel->GetVisibilityForNode(GetGuid());
@@ -476,11 +512,6 @@ EVisibility SSchematicGraphNode::GetNodeVisibility() const
 		}
 	}
 
-	if(IsBeingDragged())
-	{
-		return EVisibility::HitTestInvisible;
-	}
-
 	return EVisibility::Visible;
 }
 
@@ -846,6 +877,20 @@ void SSchematicGraphPanel::OnDropEvent(SSchematicGraphNode* Node, const FDragDro
 	OnDropDelegate.ExecuteIfBound(this, Node, InDragDropEvent);
 }
 
+FReply SSchematicGraphPanel::HandleNodeDragDetected(FGuid Guid, const FGeometry& MyGeometry, const FPointerEvent& MouseEvent)
+{
+	if(GraphData)
+	{
+		FGuid ForwardedGuid = Guid;
+		if(GraphData->GetForwardedNodeForDrag(ForwardedGuid))
+		{
+			SSchematicGraphNode* ForwardedNode = const_cast<SSchematicGraphNode*>(FindNode(ForwardedGuid));
+			return ForwardedNode->OnDragDetected(MyGeometry, MouseEvent);
+		}
+	}
+	return FReply::Unhandled();
+}
+
 FVector2d SSchematicGraphPanel::GetPositionForNode(FGuid InNodeGuid) const
 {
 	const SSchematicGraphNode* NodeWidget = FindNode(InNodeGuid);
@@ -993,6 +1038,15 @@ ESchematicGraphNodeVisibility SSchematicGraphPanel::GetVisibilityForNode(FGuid I
 	return ESchematicGraphNodeVisibility::Visible;
 }
 
+bool SSchematicGraphPanel::IsDragSupportedForNode(FGuid InNodeGuid) const
+{
+	if(GraphData)
+	{
+		return GraphData->IsDragSupportedForNode(InNodeGuid);
+	}
+	return false;
+}
+
 void SSchematicGraphPanel::UpdateAutoScalingForNodes()
 {
 	TArray<bool> NodeIsAutoScaling;
diff --git a/Engine/Source/Editor/AnimationEditorWidgets/Public/SchematicGraphPanel/SSchematicGraphPanel.h b/Engine/Source/Editor/AnimationEditorWidgets/Public/SchematicGraphPanel/SSchematicGraphPanel.h
--- a/Engine/Source/Editor/AnimationEditorWidgets/Public/SchematicGraphPanel/SSchematicGraphPanel.h
+++ b/Engine/Source/Editor/AnimationEditorWidgets/Public/SchematicGraphPanel/SSchematicGraphPanel.h
@@ -130,6 +130,8 @@ public:
 	virtual void SetPlacement(ESchematicGraphNodePlacementConstraint InPlacement) { Placement = InPlacement; }
 	virtual ESchematicGraphNodeVisibility GetVisibility() const { return Visibility; }
 	virtual void SetVisibility(ESchematicGraphNodeVisibility InVisibility) { Visibility = InVisibility; }
+	virtual bool IsDragSupported() const { return bDragSupported; }
+	virtual void SetDragSupported(bool InDragSupported) { bDragSupported = InDragSupported;}
 
 	virtual FString GetDragDropDecoratorLabel() const;
 
@@ -145,6 +147,7 @@ protected:
 	FText ToolTip = FText();
 	ESchematicGraphNodePlacementConstraint Placement = ESchematicGraphNodePlacementConstraint::Free;
 	ESchematicGraphNodeVisibility Visibility = ESchematicGraphNodeVisibility::Visible;
+	bool bDragSupported = false;
 
 	friend class FSchematicGraphModel;
 };
@@ -224,11 +227,15 @@ public:
 	virtual ESchematicGraphNodePlacementConstraint GetPlacementForNode(const FSchematicGraphNode* InNode) const;
 	ESchematicGraphNodeVisibility GetVisibilityForNode(const FGuid& InGuid) const;
 	virtual ESchematicGraphNodeVisibility GetVisibilityForNode(const FSchematicGraphNode* InNode) const;
+	bool IsDragSupportedForNode(const FGuid& InGuid) const;
+	virtual bool IsDragSupportedForNode(const FSchematicGraphNode* InNode) const;
 
 	FOnNodeAdded& OnNodeAdded() { return OnNodeAddedDelegate; }
 	FOnNodeRemoved& OnNodeRemoved() { return OnNodeRemovedDelegate; }
 	FOnGraphReset& OnGraphReset() { return OnGraphResetDelegate; }
 
+	virtual bool GetForwardedNodeForDrag(FGuid& InOutGuid) const { return false; }
+
 protected:
 
 	TArray<TSharedPtr<FSchematicGraphNode>> Nodes;
@@ -423,6 +430,7 @@ public:
 	void OnBeginDragEvent(SSchematicGraphNode* Node, const FDragDropOperation& InDragDropEvent);
 	void OnEndDragEvent(SSchematicGraphNode* Node, const FDragDropOperation& InDragDropEvent);
 	void OnDropEvent(SSchematicGraphNode* Node, const FDragDropEvent& InDragDropEvent);
+	FReply HandleNodeDragDetected(FGuid Guid, const FGeometry& MyGeometry, const FPointerEvent& MouseEvent);
 
 	virtual FVector2d GetPositionForNode(FGuid InNodeGuid) const;
 	virtual FVector2d GetSizeForNode(FGuid InNodeGuid) const;
@@ -432,6 +440,7 @@ public:
 	virtual const FSlateBrush* GetBrushForNode(FGuid InNodeGuid) const;
 	virtual FText GetToolTipForNode(FGuid InNodeGuid) const;
 	virtual ESchematicGraphNodeVisibility GetVisibilityForNode(FGuid InNodeGuid) const;
+	virtual bool IsDragSupportedForNode(FGuid InNodeGuid) const;
 
 	void UpdateAutoScalingForNodes();
 
