Control Rig: introduce the notion of links to the schematic

#rnx

diff --git a/Engine/Plugins/Animation/ControlRig/Source/ControlRigEditor/Private/EditMode/ControlRigEditModeSettings.h b/Engine/Plugins/Animation/ControlRig/Source/ControlRigEditor/Private/EditMode/ControlRigEditModeSettings.h
--- a/Engine/Plugins/Animation/ControlRig/Source/ControlRigEditor/Private/EditMode/ControlRigEditModeSettings.h
+++ b/Engine/Plugins/Animation/ControlRig/Source/ControlRigEditor/Private/EditMode/ControlRigEditModeSettings.h
@@ -18,7 +18,7 @@ class UControlRigEditModeSettings : public UObject
 	UControlRigEditModeSettings()
 		: bDisplayHierarchy(false)
 		, bDisplayNulls(false)
-		, bDisplaySockets(true)
+		, bDisplaySockets(false)
 		, bHideControlShapes(false)
 		, bShowAllProxyControls(false)
 		, bShowControlsAsOverlay(false)
diff --git a/Engine/Plugins/Animation/ControlRig/Source/ControlRigEditor/Private/Editor/ControlRigSchematicModel.cpp b/Engine/Plugins/Animation/ControlRig/Source/ControlRigEditor/Private/Editor/ControlRigSchematicModel.cpp
--- a/Engine/Plugins/Animation/ControlRig/Source/ControlRigEditor/Private/Editor/ControlRigSchematicModel.cpp
+++ b/Engine/Plugins/Animation/ControlRig/Source/ControlRigEditor/Private/Editor/ControlRigSchematicModel.cpp
@@ -64,6 +64,8 @@ FControlRigSchematicRigElementKeyNode* FControlRigSchematicModel::AddElementKeyN
 		{
 			OnNodeAddedDelegate.Broadcast(Node);
 		}
+
+		UpdateElementKeyLinks();
 	}
 	return Node;
 }
@@ -102,11 +104,180 @@ bool FControlRigSchematicModel::RemoveElementKeyNode(const FRigElementKey& InKey
 	if(const FGuid* GuidPtr = RigElementKeyToGuid.Find(InKey))
 	{
 		const FGuid Guid = *GuidPtr;
-		return RemoveNode(Guid);
+		const bool bResult = RemoveNode(Guid);
+		if(bResult)
+		{
+			UpdateElementKeyLinks();
+		}
+		return bResult;
 	}
 	return false;
 }
 
+FControlRigSchematicRigElementKeyLink* FControlRigSchematicModel::AddElementKeyLink(const FRigElementKey& InSourceKey, const FRigElementKey& InTargetKey, bool bNotify)
+{
+	check(InSourceKey.IsValid());
+	check(InTargetKey.IsValid());
+	check(InSourceKey != InTargetKey);
+	const FControlRigSchematicRigElementKeyNode* SourceNode = FindElementKeyNode(InSourceKey);
+	const FControlRigSchematicRigElementKeyNode* TargetNode = FindElementKeyNode(InTargetKey);
+	if(SourceNode && TargetNode)
+	{
+		FControlRigSchematicRigElementKeyLink* Link = AddLink<FControlRigSchematicRigElementKeyLink>(SourceNode->GetGuid(), TargetNode->GetGuid());
+		Link->SourceKey = InSourceKey;
+		Link->TargetKey = InTargetKey;
+		Link->Thickness = 2.5f;
+		return Link;
+	}
+	return nullptr;
+}
+
+void FControlRigSchematicModel::UpdateElementKeyLinks()
+{
+	const TSharedPtr<FControlRigEditor> Editor = ControlRigEditor.Pin();
+	if(!Editor.IsValid())
+	{
+		return;
+	}
+
+	typedef TTuple< FRigElementKey, FRigElementKey > TElementKeyPair;
+	typedef TTuple< FGuid, TElementKeyPair > TElementKeyLinkPair;
+
+	struct FLinkTraverser
+	{
+		const FControlRigSchematicRigElementKeyNode* VisitElement(
+			const FRigBaseElement* InElement, 
+			const FControlRigSchematicRigElementKeyNode* InNode, 
+			TArray< TElementKeyPair >& OutExpectedLinks) const
+		{
+			const FControlRigSchematicRigElementKeyNode* Node = InNode;
+			if(const FControlRigSchematicRigElementKeyNode* SelfNode = FindNode(InElement->GetKey()))
+			{
+				Node = SelfNode;
+			}
+
+			const TConstArrayView<FRigBaseElement*> Children = Hierarchy->GetChildren(InElement);
+			for(const FRigBaseElement* Child : Children)
+			{
+				const FControlRigSchematicRigElementKeyNode* ChildNode = VisitElement(Child, Node, OutExpectedLinks);
+				if(Node && ChildNode && Node != ChildNode)
+				{
+					OutExpectedLinks.Emplace(Node->GetKey(), ChildNode->GetKey());
+				}
+			}
+
+			return Node;
+		}
+
+		const FControlRigSchematicRigElementKeyNode* FindNode(const FRigElementKey& InKey) const
+		{
+			if(const FGuid* ElementGuid = RigElementKeyToGuid->Find(InKey))
+			{
+				return Cast<FControlRigSchematicRigElementKeyNode>(NodeByGuid->FindChecked(*ElementGuid).Get());
+			}
+			if(const FRigElementKey* SocketKey = SocketToParent.Find(InKey))
+			{
+				if(const FGuid* ElementGuid = RigElementKeyToGuid->Find(*SocketKey))
+				{
+					return Cast<FControlRigSchematicRigElementKeyNode>(NodeByGuid->FindChecked(*ElementGuid).Get());
+				}
+			}
+			if(const FRigElementKey* ParentKey = ParentToSocket.Find(InKey))
+			{
+				if(const FGuid* ElementGuid = RigElementKeyToGuid->Find(*ParentKey))
+				{
+					return Cast<FControlRigSchematicRigElementKeyNode>(NodeByGuid->FindChecked(*ElementGuid).Get());
+				}
+			}
+			return nullptr;
+		}
+
+		TArray< TElementKeyPair > ComputeExpectedLinks() const
+		{
+			SocketToParent.Reset();
+			ParentToSocket.Reset();
+
+			// create a map to look up sockets
+			const TArray<FRigElementKey> SocketKeys = Hierarchy->GetSocketKeys();
+			for(const FRigElementKey& SocketKey : SocketKeys)
+			{
+				const FRigElementKey ParentKey = Hierarchy->GetFirstParent(SocketKey);
+				if(ParentKey.IsValid())
+				{
+					SocketToParent.Add(SocketKey, ParentKey);
+					if(!ParentToSocket.Contains(ParentKey))
+					{
+						ParentToSocket.Add(ParentKey, SocketKey);
+					}
+				}
+			}
+			
+			TArray< TElementKeyPair > ExpectedLinks;
+			const TArray<FRigBaseElement*> RootElements = Hierarchy->GetRootElements();
+			for(const FRigBaseElement* RootElement : RootElements)
+			{
+				VisitElement(RootElement, nullptr, ExpectedLinks);
+			}
+			return ExpectedLinks;
+		}
+
+		const UControlRig* ControlRig; 
+		const URigHierarchy* Hierarchy; 
+		const TMap<FRigElementKey, FGuid>* RigElementKeyToGuid = nullptr;
+		const TMap<FGuid, TSharedPtr<FSchematicGraphNode>>* NodeByGuid = nullptr;
+		mutable TMap<FRigElementKey, FRigElementKey> SocketToParent;
+		mutable TMap<FRigElementKey, FRigElementKey> ParentToSocket;
+	};
+
+	FLinkTraverser Traverser;
+	Traverser.ControlRig = Editor->GetControlRig();
+	if(Traverser.ControlRig == nullptr)
+	{
+		return;
+	}
+	Traverser.Hierarchy = Traverser.ControlRig->GetHierarchy();
+	if(Traverser.Hierarchy == nullptr)
+	{
+		return;
+	}
+	Traverser.RigElementKeyToGuid = &RigElementKeyToGuid;
+	Traverser.NodeByGuid = &NodeByGuid;
+
+	const TArray< TElementKeyPair > ExpectedLinks = Traverser.ComputeExpectedLinks();
+
+	TArray< TElementKeyLinkPair > ExistingLinks;
+	for(const TSharedPtr<FSchematicGraphLink>& Link : Links)
+	{
+		if(const FControlRigSchematicRigElementKeyLink* ElementKeyLink = Cast<FControlRigSchematicRigElementKeyLink>(Link.Get()))
+		{
+			ExistingLinks.Emplace(ElementKeyLink->GetGuid(), TElementKeyPair(ElementKeyLink->GetSourceKey(), ElementKeyLink->GetTargetKey()));
+		}
+	}
+
+	// remove the obsolete links
+	for(const TTuple< FGuid, TTuple< FRigElementKey, FRigElementKey > >& ExistingLink : ExistingLinks)
+	{
+		if(!ExpectedLinks.Contains(ExistingLink.Get<1>()))
+		{
+			(void)RemoveLink(ExistingLink.Get<0>());
+		}
+	}
+
+	// add missing links
+	for(const TElementKeyPair& ExpectedLink : ExpectedLinks)
+	{
+		if(!ExistingLinks.ContainsByPredicate([ExpectedLink](const TElementKeyLinkPair& ExistingLink) -> bool
+		{
+			return ExistingLink.Get<1>() == ExpectedLink;
+		}))
+		{
+			(void)AddElementKeyLink(ExpectedLink.Get<0>(), ExpectedLink.Get<1>());
+		}
+	}
+
+	// todo: also introduce links for module relationships
+}
+
 void FControlRigSchematicModel::OnSetObjectBeingDebugged(UObject* InObject)
 {
 	if(ControlRigBeingDebuggedPtr.Get() == InObject)
@@ -135,6 +306,9 @@ void FControlRigSchematicModel::OnSetObjectBeingDebugged(UObject* InObject)
 			Hierarchy->OnModified().RemoveAll(this);
 			Hierarchy->OnModified().AddRaw(this, &FControlRigSchematicModel::OnHierarchyModified);
 		}
+
+		// todo: react to the rig being constructed - but not when the rig hierarchy is being copied 
+		// ControlRig->OnPostConstruction_AnyThread().AddRaw(this, &FControlRigSchematicModel::OnPostConstruction);
 	}
 }
 
@@ -396,11 +570,13 @@ FLinearColor FControlRigSchematicModel::GetColorForNode(const FSchematicGraphNod
 			}
 			case ERigElementType::Socket:
 			{
-				return FControlRigEditorStyle::Get().SocketUserInterfaceColor;
+				return FLinearColor::White;
+				//return FControlRigEditorStyle::Get().SocketUserInterfaceColor;
 			}
 			case ERigElementType::Connector:
 			{
-				return FControlRigEditorStyle::Get().ConnectorUserInterfaceColor;
+				return FLinearColor::White;
+				//return FControlRigEditorStyle::Get().ConnectorUserInterfaceColor;
 			}
 		}
 	}
@@ -441,7 +617,7 @@ const FText FControlRigSchematicModel::GetToolTipForNode(const FSchematicGraphNo
 	return FSchematicGraphModel::GetToolTipForNode(InNode);
 }
 
-ESchematicGraphNodeVisibility FControlRigSchematicModel::GetVisibilityForNode(const FSchematicGraphNode* InNode) const
+ESchematicGraphVisibility FControlRigSchematicModel::GetVisibilityForNode(const FSchematicGraphNode* InNode) const
 {
 	if(const FControlRigSchematicRigElementKeyNode* Node = Cast<FControlRigSchematicRigElementKeyNode>(InNode))
 	{
@@ -452,7 +628,7 @@ ESchematicGraphNodeVisibility FControlRigSchematicModel::GetVisibilityForNode(co
 			{
 				if(ContainsElementKeyNode(ResolvedSocket))
 				{
-					return ESchematicGraphNodeVisibility::Hidden;
+					return ESchematicGraphVisibility::Hidden;
 				}
 			}
 		}
diff --git a/Engine/Plugins/Animation/ControlRig/Source/ControlRigEditor/Private/Editor/ControlRigSchematicModel.h b/Engine/Plugins/Animation/ControlRig/Source/ControlRigEditor/Private/Editor/ControlRigSchematicModel.h
--- a/Engine/Plugins/Animation/ControlRig/Source/ControlRigEditor/Private/Editor/ControlRigSchematicModel.h
+++ b/Engine/Plugins/Animation/ControlRig/Source/ControlRigEditor/Private/Editor/ControlRigSchematicModel.h
@@ -33,6 +33,27 @@ protected:
 	friend class FControlRigSchematicModel;
 };
 
+/** Link between two rig element schematic nodes */
+class FControlRigSchematicRigElementKeyLink : public FSchematicGraphLink
+{
+public:
+
+	SCHEMATICGRAPHLINK_BODY(FControlRigSchematicRigElementKeyLink, FSchematicGraphLink)
+
+	virtual ~FControlRigSchematicRigElementKeyLink() override {}
+	
+	const FRigElementKey& GetSourceKey() const { return SourceKey; }
+	const FRigElementKey& GetTargetKey() const { return TargetKey; }
+	
+protected:
+
+	FRigElementKey SourceKey;
+	FRigElementKey TargetKey;
+
+	friend class FControlRigSchematicModel;
+};
+
+
 /** Model for the schematic views */
 class FControlRigSchematicModel : public FSchematicGraphModel
 {
@@ -47,6 +68,8 @@ public:
 	bool ContainsElementKeyNode(const FRigElementKey& InKey) const;
 	virtual bool RemoveNode(const FGuid& InGuid) override;
 	bool RemoveElementKeyNode(const FRigElementKey& InKey);
+	FControlRigSchematicRigElementKeyLink* AddElementKeyLink(const FRigElementKey& InSourceKey, const FRigElementKey& InTargetKey, bool bNotify = true);
+	void UpdateElementKeyLinks();
 
 	void OnSetObjectBeingDebugged(UObject* InObject);
 	void OnHierarchyModified(ERigHierarchyNotification InNotif, URigHierarchy* InHierarchy, const FRigBaseElement* InElement);
@@ -57,7 +80,7 @@ public:
 	virtual const FSlateBrush* GetBrushForNode(const FSchematicGraphNode* InNode) const override;
 	virtual FLinearColor GetColorForNode(const FSchematicGraphNode* InNode) const override;
 	virtual const FText GetToolTipForNode(const FSchematicGraphNode* InNode) const override;
-	virtual ESchematicGraphNodeVisibility GetVisibilityForNode(const FSchematicGraphNode* InNode) const override;
+	virtual ESchematicGraphVisibility GetVisibilityForNode(const FSchematicGraphNode* InNode) const override;
 	virtual ESchematicGraphNodePlacementConstraint GetPlacementForNode(const FSchematicGraphNode* InNode) const override;
 
 	virtual bool GetForwardedNodeForDrag(FGuid& InOutGuid) const override;
diff --git a/Engine/Source/Editor/AnimationEditorWidgets/Private/SchematicGraphPanel/SSchematicGraphPanel.cpp b/Engine/Source/Editor/AnimationEditorWidgets/Private/SchematicGraphPanel/SSchematicGraphPanel.cpp
--- a/Engine/Source/Editor/AnimationEditorWidgets/Private/SchematicGraphPanel/SSchematicGraphPanel.cpp
+++ b/Engine/Source/Editor/AnimationEditorWidgets/Private/SchematicGraphPanel/SSchematicGraphPanel.cpp
@@ -1,4 +1,5 @@
 // Copyright Epic Games, Inc. All Rights Reserved.
+
 #if WITH_EDITOR
 
 #include "SchematicGraphPanel/SSchematicGraphPanel.h"
@@ -22,7 +23,17 @@ FString FSchematicGraphNode::GetDragDropDecoratorLabel() const
 	return GetGuid().ToString();
 }
 
-void FSchematicGraphModel::Reset()
+ const FName& FSchematicGraphLink::GetType() const
+ {
+	return FSchematicGraphLink::Type;
+ }
+
+ bool FSchematicGraphLink::IsA(const FName& InType) const
+ {
+	return Type == InType;
+ }
+
+ void FSchematicGraphModel::Reset()
 {
 	Nodes.Reset();
 	Links.Reset();
@@ -33,10 +44,25 @@ void FSchematicGraphModel::Reset()
 	}
 }
 
-bool FSchematicGraphModel::RemoveNode(const FGuid& InGuid)
+bool FSchematicGraphModel::RemoveNode(const FGuid& InNodeGuid)
 {
-	if(const FSchematicGraphNode* Node = FindNode(InGuid))
+	if(const FSchematicGraphNode* Node = FindNode(InNodeGuid))
 	{
+		if(TTuple<TArray<FGuid>,TArray<FGuid>>* ExistingTuple = NodeGuidToLinkGuids.Find(InNodeGuid))
+		{
+			TTuple<TArray<FGuid>,TArray<FGuid>> LinksToNode;
+			Swap(LinksToNode, *ExistingTuple);
+			NodeGuidToLinkGuids.Remove(InNodeGuid);
+			for(const FGuid& LinkGuid : LinksToNode.Get<0>())
+			{
+				(void)RemoveLink(LinkGuid);
+			}
+			for(const FGuid& LinkGuid : LinksToNode.Get<1>())
+			{
+				(void)RemoveLink(LinkGuid);
+			}
+		}
+		
 		if (OnNodeRemovedDelegate.IsBound())
 		{
 			OnNodeRemovedDelegate.Broadcast(Node);
@@ -52,9 +78,9 @@ bool FSchematicGraphModel::RemoveNode(const FGuid& InGuid)
 	return false;
 }
 
-FVector2d FSchematicGraphModel::GetPositionForNode(const FGuid& InGuid) const
+FVector2d FSchematicGraphModel::GetPositionForNode(const FGuid& InNodeGuid) const
 {
-	if(const FSchematicGraphNode* Node = FindNode(InGuid))
+	if(const FSchematicGraphNode* Node = FindNode(InNodeGuid))
 	{
 		return GetPositionForNode(Node);
 	}
@@ -67,9 +93,9 @@ FVector2d FSchematicGraphModel::GetPositionForNode(const FSchematicGraphNode* In
 	return InNode->GetPosition();
 }
 
-FVector2d FSchematicGraphModel::GetPositionOffsetForNode(const FGuid& InGuid) const
+FVector2d FSchematicGraphModel::GetPositionOffsetForNode(const FGuid& InNodeGuid) const
 {
-	if(const FSchematicGraphNode* Node = FindNode(InGuid))
+	if(const FSchematicGraphNode* Node = FindNode(InNodeGuid))
 	{
 		return GetPositionOffsetForNode(Node);
 	}
@@ -102,13 +128,12 @@ FVector2d FSchematicGraphModel::GetSizeForNode(const FGuid& InNodeGuid) const
 	{
 		return GetSizeForNode(Node);
 	}
-	return FVector2d::ZeroVector;
+	return SSchematicGraphNode::DefaultNodeSize;
 }
 
 FVector2d FSchematicGraphModel::GetSizeForNode(const FSchematicGraphNode* InNode) const
 {
-	static const FVector2d DefaultSize = FVector2d(50.0,50.0);
-	return DefaultSize;
+	return SSchematicGraphNode::DefaultNodeSize;
 }
 
 float FSchematicGraphModel::GetScaleForNode(const FGuid& InNodeGuid, bool bIncludeScaleOffset) const
@@ -144,7 +169,25 @@ float FSchematicGraphModel::GetScaleOffsetForNode(const FSchematicGraphNode* InN
 	return InNode->GetScaleOffset();
 }
 
-bool FSchematicGraphModel::IsAutoScaleEnabledForNode(const FGuid& InNodeGuid) const
+ float FSchematicGraphModel::GetMinimumLinkDistanceForNode(const FGuid& InNodeGuid) const
+ {
+	if(const FSchematicGraphNode* Node = FindNode(InNodeGuid))
+	{
+		return GetMinimumLinkDistanceForNode(Node);
+	}
+	return 0;
+ }
+
+ float FSchematicGraphModel::GetMinimumLinkDistanceForNode(const FSchematicGraphNode* InNode) const
+ {
+	check(InNode);
+
+	// return the radius of the node - 1 pixel
+	const FVector2d Size = GetSizeForNode(InNode);
+	return Size.GetMax() * 0.5f - 1.f;
+ }
+
+ bool FSchematicGraphModel::IsAutoScaleEnabledForNode(const FGuid& InNodeGuid) const
 {
 	if(const FSchematicGraphNode* Node = FindNode(InNodeGuid))
 	{
@@ -159,12 +202,12 @@ bool FSchematicGraphModel::IsAutoScaleEnabledForNode(const FSchematicGraphNode*
 	return InNode->IsAutoScaleEnabled();
 }
 
-FLinearColor FSchematicGraphModel::GetColorForNode(const FGuid& InGuid) const
+FLinearColor FSchematicGraphModel::GetColorForNode(const FGuid& InNodeGuid) const
 {
-	if(const FSchematicGraphNode* Node = FindNode(InGuid))
+	if(const FSchematicGraphNode* Node = FindNode(InNodeGuid))
 	{
 		const FLinearColor Color = GetColorForNode(Node);
-		if(GetVisibilityForNode(Node) == ESchematicGraphNodeVisibility::FadedOut)
+		if(GetVisibilityForNode(Node) == ESchematicGraphVisibility::FadedOut)
 		{
 			return Color * 0.5f;
 		}
@@ -179,9 +222,9 @@ FLinearColor FSchematicGraphModel::GetColorForNode(const FSchematicGraphNode* In
 	return InNode->GetColor();
 }
 
-const FSlateBrush* FSchematicGraphModel::GetBrushForNode(const FGuid& InGuid) const
+const FSlateBrush* FSchematicGraphModel::GetBrushForNode(const FGuid& InNodeGuid) const
 {
-	if(const FSchematicGraphNode* Node = FindNode(InGuid))
+	if(const FSchematicGraphNode* Node = FindNode(InNodeGuid))
 	{
 		if(const FSlateBrush* Brush = GetBrushForNode(Node))
 		{
@@ -199,9 +242,9 @@ const FSlateBrush* FSchematicGraphModel::GetBrushForNode(const FSchematicGraphNo
 	return InNode->GetBrush();
 }
 
-const FText FSchematicGraphModel::GetToolTipForNode(const FGuid& InGuid) const
+const FText FSchematicGraphModel::GetToolTipForNode(const FGuid& InNodeGuid) const
 {
-	if(const FSchematicGraphNode* Node = FindNode(InGuid))
+	if(const FSchematicGraphNode* Node = FindNode(InNodeGuid))
 	{
 		return GetToolTipForNode(Node);
 	}
@@ -214,9 +257,9 @@ const FText FSchematicGraphModel::GetToolTipForNode(const FSchematicGraphNode* I
 	return InNode->GetToolTip();
 }
 
-ESchematicGraphNodePlacementConstraint FSchematicGraphModel::GetPlacementForNode(const FGuid& InGuid) const
+ESchematicGraphNodePlacementConstraint FSchematicGraphModel::GetPlacementForNode(const FGuid& InNodeGuid) const
 {
-	if(const FSchematicGraphNode* Node = FindNode(InGuid))
+	if(const FSchematicGraphNode* Node = FindNode(InNodeGuid))
 	{
 		return GetPlacementForNode(Node);
 	}
@@ -229,24 +272,24 @@ ESchematicGraphNodePlacementConstraint FSchematicGraphModel::GetPlacementForNode
 	return InNode->GetPlacement();
 }
 
-ESchematicGraphNodeVisibility FSchematicGraphModel::GetVisibilityForNode(const FGuid& InGuid) const
+ESchematicGraphVisibility FSchematicGraphModel::GetVisibilityForNode(const FGuid& InNodeGuid) const
 {
-	if(const FSchematicGraphNode* Node = FindNode(InGuid))
+	if(const FSchematicGraphNode* Node = FindNode(InNodeGuid))
 	{
 		return GetVisibilityForNode(Node);
 	}
-	return ESchematicGraphNodeVisibility::Visible;
+	return ESchematicGraphVisibility::Visible;
 }
 
-ESchematicGraphNodeVisibility FSchematicGraphModel::GetVisibilityForNode(const FSchematicGraphNode* InNode) const
+ESchematicGraphVisibility FSchematicGraphModel::GetVisibilityForNode(const FSchematicGraphNode* InNode) const
 {
 	check(InNode);
 	return InNode->GetVisibility();
 }
 
-bool FSchematicGraphModel::IsDragSupportedForNode(const FGuid& InGuid) const
+bool FSchematicGraphModel::IsDragSupportedForNode(const FGuid& InNodeGuid) const
 {
-	if(const FSchematicGraphNode* Node = FindNode(InGuid))
+	if(const FSchematicGraphNode* Node = FindNode(InNodeGuid))
 	{
 		return IsDragSupportedForNode(Node);
 	}
@@ -259,6 +302,231 @@ bool FSchematicGraphModel::IsDragSupportedForNode(const FSchematicGraphNode* InN
 	return InNode->IsDragSupported();
 }
 
+bool FSchematicGraphModel::IsLinkedTo(const FGuid& InSourceNodeGuid, const FGuid& InTargetNodeGuid) const
+{
+	return FindLink<>(InSourceNodeGuid, InTargetNodeGuid) != nullptr;
+}
+
+TArray<const FSchematicGraphLink*> FSchematicGraphModel::FindLinksOnNode(const FGuid& InNodeGuid) const
+{
+	TArray<const FSchematicGraphLink*> Result;
+	for(const TSharedPtr<FSchematicGraphLink>& Link : Links)
+	{
+		if(Link->GetSourceNodeGuid() == InNodeGuid ||
+			Link->GetTargetNodeGuid() == InNodeGuid)
+		{
+			Result.Add(Link.Get());
+		}
+	}
+	return Result;
+}
+
+TArray<const FSchematicGraphLink*> FSchematicGraphModel::FindLinksOnSource(const FGuid& InSourceNodeGuid) const
+{
+	TArray<const FSchematicGraphLink*> Result;
+	for(const TSharedPtr<FSchematicGraphLink>& Link : Links)
+	{
+		if(Link->GetSourceNodeGuid() == InSourceNodeGuid)
+		{
+			Result.Add(Link.Get());
+		}
+	}
+	return Result;
+}
+
+TArray<const FSchematicGraphLink*> FSchematicGraphModel::FindLinksOnTarget(const FGuid& InTargetNodeGuid) const
+{
+	TArray<const FSchematicGraphLink*> Result;
+	for(const TSharedPtr<FSchematicGraphLink>& Link : Links)
+	{
+		if(Link->GetTargetNodeGuid() == InTargetNodeGuid)
+		{
+			Result.Add(Link.Get());
+		}
+	}
+	return Result;
+}
+
+bool FSchematicGraphModel::RemoveLink(const FGuid& InLinkGuid)
+{
+	if(const FSchematicGraphLink* Link = FindLink(InLinkGuid))
+	{
+		if (OnLinkRemovedDelegate.IsBound())
+		{
+			OnLinkRemovedDelegate.Broadcast(Link);
+		}
+		const FGuid Guid = Link->GetGuid();
+		const uint32 LinkHash = Link->GetLinkHash();
+		const FGuid SourceNodeGuid = Link->GetSourceNodeGuid();
+		const FGuid TargetNodeGuid = Link->GetTargetNodeGuid();
+		LinkByGuid.Remove(Link->GetGuid());
+		LinkByHash.Remove(LinkHash);
+		if(TTuple<TArray<FGuid>,TArray<FGuid>>* Tuple = NodeGuidToLinkGuids.Find(SourceNodeGuid))
+		{
+			Tuple->Get<0>().Remove(Guid);
+			if(Tuple->Get<0>().IsEmpty() && Tuple->Get<1>().IsEmpty())
+			{
+				NodeGuidToLinkGuids.Remove(SourceNodeGuid);
+			}
+		}
+		if(TTuple<TArray<FGuid>,TArray<FGuid>>* Tuple = NodeGuidToLinkGuids.Find(TargetNodeGuid))
+		{
+			Tuple->Get<1>().Remove(Guid);
+			if(Tuple->Get<0>().IsEmpty() && Tuple->Get<1>().IsEmpty())
+			{
+				NodeGuidToLinkGuids.Remove(TargetNodeGuid);
+			}
+		}
+
+		Links.RemoveAll([Guid](const TSharedPtr<FSchematicGraphLink>& ExistingLink) -> bool
+		{
+			return ExistingLink->GetGuid() == Guid;
+		});
+		return true;
+	}
+	return false;
+}
+
+ float FSchematicGraphModel::GetMinimumForLink(const FGuid& InLinkGuid) const
+ {
+	if(const FSchematicGraphLink* Link = FindLink(InLinkGuid))
+	{
+		return GetMinimumForLink(Link);
+	}
+	return 0.f;
+ }
+
+ float FSchematicGraphModel::GetMinimumForLink(const FSchematicGraphLink* InLink) const
+ {
+	check(InLink);
+	return InLink->GetMinimum();
+ }
+
+ float FSchematicGraphModel::GetMaximumForLink(const FGuid& InLinkGuid) const
+ {
+	if(const FSchematicGraphLink* Link = FindLink(InLinkGuid))
+	{
+		return GetMaximumForLink(Link);
+	}
+	return 1.f;
+ }
+
+ float FSchematicGraphModel::GetMaximumForLink(const FSchematicGraphLink* InLink) const
+ {
+	check(InLink);
+	return InLink->GetMaximum();
+ }
+
+FVector2d FSchematicGraphModel::GetSourceNodeOffsetForLink(const FGuid& InLinkGuid) const
+ {
+	if(const FSchematicGraphLink* Link = FindLink(InLinkGuid))
+	{
+		return GetSourceNodeOffsetForLink(Link);
+	}
+	return FVector2d::ZeroVector;
+ }
+
+ FVector2d FSchematicGraphModel::GetSourceNodeOffsetForLink(const FSchematicGraphLink* InLink) const
+ {
+	check(InLink);
+	return InLink->GetSourceNodeOffset();
+ }
+
+ FVector2d FSchematicGraphModel::GetTargetNodeOffsetForLink(const FGuid& InLinkGuid) const
+ {
+	if(const FSchematicGraphLink* Link = FindLink(InLinkGuid))
+	{
+		return GetTargetNodeOffsetForLink(Link);
+	}
+	return FVector2d::ZeroVector;
+ }
+
+ FVector2d FSchematicGraphModel::GetTargetNodeOffsetForLink(const FSchematicGraphLink* InLink) const
+ {
+	check(InLink);
+	return InLink->GetTargetNodeOffset();
+ }
+
+ FLinearColor FSchematicGraphModel::GetColorForLink(const FGuid& InLinkGuid) const
+{
+	if(const FSchematicGraphLink* Link = FindLink(InLinkGuid))
+	{
+		const FLinearColor Color = GetColorForLink(Link);
+		if(GetVisibilityForLink(Link) == ESchematicGraphVisibility::FadedOut)
+		{
+			return Color * 0.5f;
+		}
+		return Color;
+	}
+	return FLinearColor::White;
+}
+
+FLinearColor FSchematicGraphModel::GetColorForLink(const FSchematicGraphLink* InLink) const
+{
+	check(InLink);
+	return InLink->GetColor();
+}
+
+ float FSchematicGraphModel::GetThicknessForLink(const FGuid& InLinkGuid) const
+ {
+	if(const FSchematicGraphLink* Link = FindLink(InLinkGuid))
+	{
+		return GetThicknessForLink(Link);
+	}
+	return 1.f;
+ }
+
+ float FSchematicGraphModel::GetThicknessForLink(const FSchematicGraphLink* InLink) const
+ {
+	check(InLink);
+	return InLink->GetThickness();
+ }
+
+ const FSlateBrush* FSchematicGraphModel::GetBrushForLink(const FGuid& InLinkGuid) const
+{
+	if(const FSchematicGraphLink* Link = FindLink(InLinkGuid))
+	{
+		return GetBrushForLink(Link);
+	}
+	return nullptr;
+}
+
+const FSlateBrush* FSchematicGraphModel::GetBrushForLink(const FSchematicGraphLink* InLink) const
+{
+	check(InLink);
+	return InLink->GetBrush();
+}
+
+const FText FSchematicGraphModel::GetToolTipForLink(const FGuid& InLinkGuid) const
+{
+	if(const FSchematicGraphLink* Link = FindLink(InLinkGuid))
+	{
+		return GetToolTipForLink(Link);
+	}
+	return FText();
+}
+
+const FText FSchematicGraphModel::GetToolTipForLink(const FSchematicGraphLink* InLink) const
+{
+	check(InLink);
+	return InLink->GetToolTip();
+}
+
+ESchematicGraphVisibility FSchematicGraphModel::GetVisibilityForLink(const FGuid& InLinkGuid) const
+{
+	if(const FSchematicGraphLink* Link = FindLink(InLinkGuid))
+	{
+		return GetVisibilityForLink(Link);
+	}
+	return ESchematicGraphVisibility::Visible;
+}
+
+ESchematicGraphVisibility FSchematicGraphModel::GetVisibilityForLink(const FSchematicGraphLink* InLink) const
+{
+	check(InLink);
+	return InLink->GetVisibility();
+}
+
 TSharedRef<FSchematicGraphNodeDragDropOp> FSchematicGraphNodeDragDropOp::New(TArray<SSchematicGraphNode*> InSchematicGraphNodes, const TArray<FGuid>& InElements, FSchematicGraphNodeDragDropOp::FOnEndDrag InOnEndDragDelegate)
 {
 	TSharedRef<FSchematicGraphNodeDragDropOp> Operation = MakeShared<FSchematicGraphNodeDragDropOp>();
@@ -501,12 +769,12 @@ EVisibility SSchematicGraphNode::GetNodeVisibility() const
 
 	if(SchematicGraphPanel)
 	{
-		ESchematicGraphNodeVisibility Vis = SchematicGraphPanel->GetVisibilityForNode(GetGuid());
-		if(Vis == ESchematicGraphNodeVisibility::Hidden)
+		ESchematicGraphVisibility Vis = SchematicGraphPanel->GetVisibilityForNode(GetGuid());
+		if(Vis == ESchematicGraphVisibility::Hidden)
 		{
 			return EVisibility::Hidden;
 		}
-		if(Vis == ESchematicGraphNodeVisibility::FadedOut)
+		if(Vis == ESchematicGraphVisibility::FadedOut)
 		{
 			return EVisibility::HitTestInvisible;
 		}
@@ -544,7 +812,7 @@ bool SSchematicGraphNode::IsInteractive() const
 {
 	if(NodeData)
 	{
-		return NodeData->GetVisibility() == ESchematicGraphNodeVisibility::Visible;
+		return NodeData->GetVisibility() == ESchematicGraphVisibility::Visible;
 	}
 	return true;
 }
@@ -555,6 +823,8 @@ void SSchematicGraphPanel::SetSchematicGraph(FSchematicGraphModel* InGraphData)
 	{
 		GraphData->OnNodeAdded().RemoveAll(this);
 		GraphData->OnNodeRemoved().RemoveAll(this);
+		GraphData->OnLinkAdded().RemoveAll(this);
+		GraphData->OnLinkRemoved().RemoveAll(this);
 		GraphData->OnGraphReset().RemoveAll(this);
 	}
 	
@@ -563,6 +833,8 @@ void SSchematicGraphPanel::SetSchematicGraph(FSchematicGraphModel* InGraphData)
 	{
 		GraphData->OnNodeAdded().AddSP(this, &SSchematicGraphPanel::AddNode);
 		GraphData->OnNodeRemoved().AddSP(this, &SSchematicGraphPanel::RemoveNode);
+		GraphData->OnLinkAdded().AddSP(this, &SSchematicGraphPanel::AddLink);
+		GraphData->OnLinkRemoved().AddSP(this, &SSchematicGraphPanel::RemoveLink);
 		GraphData->OnGraphReset().AddSP(this, &SSchematicGraphPanel::RebuildPanel);
 	}
 }
@@ -581,8 +853,8 @@ void SSchematicGraphPanel::Construct(const FArguments& InArgs)
 	OnEndDragDelegate = InArgs._OnEndDrag;
 	OnDropDelegate = InArgs._OnDrop;
 
-	TEasingAttributeInterpolator<float>::FSettings FloatInterpSettings(EEasingInterpolatorType::CubicEaseOut, 0.35f);
-	FadeBackgroundAlpha = FFloatAttribute::CreateWithGetter(FloatInterpSettings, FFloatAttribute::FGetter::CreateLambda([this]() -> float
+	TEasingAttributeInterpolator<float>::FSettings FloatInterpolationSettings(EEasingInterpolatorType::CubicEaseOut, 0.35f);
+	FadeBackgroundAlpha = FFloatAttribute::CreateWithGetter(FloatInterpolationSettings, FFloatAttribute::FGetter::CreateLambda([this]() -> float
 	{
 		return bIsDragDropping ? 0.5f : 0.f;
 	}));
@@ -594,6 +866,8 @@ void SSchematicGraphPanel::Construct(const FArguments& InArgs)
 	{
 		GraphData->OnNodeAdded().AddSP(this, &SSchematicGraphPanel::AddNode);
 		GraphData->OnNodeRemoved().AddSP(this, &SSchematicGraphPanel::RemoveNode);
+		GraphData->OnLinkAdded().AddSP(this, &SSchematicGraphPanel::AddLink);
+		GraphData->OnLinkRemoved().AddSP(this, &SSchematicGraphPanel::RemoveLink);
 		GraphData->OnGraphReset().AddSP(this, &SSchematicGraphPanel::RebuildPanel);
 	}
 }
@@ -608,20 +882,24 @@ void SSchematicGraphPanel::RebuildPanel()
 		{
 			AddNode(Node.Get());
 		}
+		for (const TSharedPtr<FSchematicGraphLink>& Link : GraphData->GetLinks())
+		{
+			AddLink(Link.Get());
+		}
 	}
 }
 
 void SSchematicGraphPanel::AddNode(const FSchematicGraphNode* InNodeToAdd)
 {
-	TEasingAttributeInterpolator<FVector2d>::FSettings Vector2DInterpSettings(EEasingInterpolatorType::CubicEaseOut, 0.2f);
-	TEasingAttributeInterpolator<float>::FSettings FloatInterpSettings(EEasingInterpolatorType::CubicEaseOut, 0.2f);
-	TEasingAttributeInterpolator<FLinearColor>::FSettings ColorInterpSettings(EEasingInterpolatorType::CubicEaseOut, 0.2f);
+	static const TEasingAttributeInterpolator<FVector2d>::FSettings Vector2DInterpolationSettings(EEasingInterpolatorType::CubicEaseOut, 0.2f);
+	static const TEasingAttributeInterpolator<float>::FSettings FloatInterpolationSettings(EEasingInterpolatorType::CubicEaseOut, 0.2f);
+	static const TEasingAttributeInterpolator<FLinearColor>::FSettings ColorInterpolationSettings(EEasingInterpolatorType::CubicEaseOut, 0.2f);
 
 	const FGuid Guid = InNodeToAdd->GetGuid();
-	const auto Position = FVector2dAttribute::CreateWithGetter(Vector2DInterpSettings, FVector2dAttribute::FGetter::CreateSP(this, &SSchematicGraphPanel::GetPositionForNode, Guid));
-	const auto Size = FVector2dAttribute::CreateWithGetter(Vector2DInterpSettings, FVector2dAttribute::FGetter::CreateSP(this, &SSchematicGraphPanel::GetSizeForNode, Guid));
-	const auto Scale = FFloatAttribute::CreateWithGetter(FloatInterpSettings, FFloatAttribute::FGetter::CreateSP(this, &SSchematicGraphPanel::GetScaleForNode, Guid, true), 0.f);
-	const auto Color = FLinearColorAttribute::CreateWithGetter(ColorInterpSettings, FLinearColorAttribute::FGetter::CreateSP(this, &SSchematicGraphPanel::GetColorForNode, Guid));
+	const auto Position = FVector2dAttribute::CreateWithGetter(Vector2DInterpolationSettings, FVector2dAttribute::FGetter::CreateSP(this, &SSchematicGraphPanel::GetPositionForNode, Guid));
+	const auto Size = FVector2dAttribute::CreateWithGetter(Vector2DInterpolationSettings, FVector2dAttribute::FGetter::CreateSP(this, &SSchematicGraphPanel::GetSizeForNode, Guid));
+	const auto Scale = FFloatAttribute::CreateWithGetter(FloatInterpolationSettings, FFloatAttribute::FGetter::CreateSP(this, &SSchematicGraphPanel::GetScaleForNode, Guid, true), 0.f);
+	const auto Color = FLinearColorAttribute::CreateWithGetter(ColorInterpolationSettings, FLinearColorAttribute::FGetter::CreateSP(this, &SSchematicGraphPanel::GetColorForNode, Guid));
 
 	const TSharedRef<SSchematicGraphNode> NewNode = SNew(SSchematicGraphNode)
 														.Position(Position)
@@ -675,14 +953,51 @@ const SSchematicGraphNode* SSchematicGraphPanel::FindNode(const FGuid& InGuid) c
 	return nullptr;
 }
 
-void SSchematicGraphPanel::OnArrangeChildren(const FGeometry& AllottedGeometry, FArrangedChildren& ArrangedChildren) const
+ void SSchematicGraphPanel::AddLink(const FSchematicGraphLink* InLinkToAdd)
+ {
+	static const TEasingAttributeInterpolator<float>::FSettings FloatInterpolationSettings(EEasingInterpolatorType::CubicEaseOut, 0.1f);
+	static const TEasingAttributeInterpolator<float>::FSettings SlowFloatInterpolationSettings(EEasingInterpolatorType::CubicEaseOut, 0.2f);
+	static const TEasingAttributeInterpolator<FLinearColor>::FSettings ColorInterpolationSettings(EEasingInterpolatorType::CubicEaseOut, 0.2f);
+
+	const FGuid Guid = InLinkToAdd->GetGuid();
+	const auto Minimum = FFloatAttribute::CreateWithGetter(SlowFloatInterpolationSettings, FFloatAttribute::FGetter::CreateSP(this, &SSchematicGraphPanel::GetMinimumForLink, Guid), 0.5f);
+	const auto Maximum = FFloatAttribute::CreateWithGetter(SlowFloatInterpolationSettings, FFloatAttribute::FGetter::CreateSP(this, &SSchematicGraphPanel::GetMaximumForLink, Guid), 0.5f);
+	const auto Color = FLinearColorAttribute::CreateWithGetter(ColorInterpolationSettings, FLinearColorAttribute::FGetter::CreateSP(this, &SSchematicGraphPanel::GetColorForLink, Guid));
+	const auto Thickness = FFloatAttribute::CreateWithGetter(FloatInterpolationSettings, FFloatAttribute::FGetter::CreateSP(this, &SSchematicGraphPanel::GetThicknessForLink, Guid), 0);
+
+	FSchematicLinkWidgetInfo Info;
+	Info.Minimum = Minimum;
+	Info.Maximum = Maximum;
+	Info.Color = Color;
+	Info.Thickness = Thickness;
+
+	LinkByGuid.Add(Guid, MakeShareable(new FSchematicLinkWidgetInfo(Info)));
+ }
+
+ void SSchematicGraphPanel::RemoveLink(const FSchematicGraphLink* InLinkToRemove)
+ {
+	const FGuid GuidToRemove = InLinkToRemove->GetGuid();
+	LinkByGuid.Remove(GuidToRemove);
+ }
+
+ const SSchematicGraphPanel::FSchematicLinkWidgetInfo* SSchematicGraphPanel::FindLink(const FGuid& InGuid) const
+ {
+	if(const TSharedPtr<FSchematicLinkWidgetInfo>* FoundLinkPtr = LinkByGuid.Find(InGuid))
+	{
+		return FoundLinkPtr->Get();
+	}
+	return nullptr;
+ }
+
+ void SSchematicGraphPanel::OnArrangeChildren(const FGeometry& AllottedGeometry, FArrangedChildren& ArrangedChildren) const
 {
 	SNodePanel::OnArrangeChildren(AllottedGeometry, ArrangedChildren);
 }
 
 int32 SSchematicGraphPanel::OnPaint(const FPaintArgs& Args, const FGeometry& AllottedGeometry, const FSlateRect& MyCullingRect, FSlateWindowElementList& OutDrawElements, int32 LayerId, const FWidgetStyle& InWidgetStyle, bool bParentEnabled) const
 {
-	const int32 NodeLayerId = LayerId + 1;
+	const int32 LinkLayerId = LayerId + 1;
+	const int32 NodeLayerId = LinkLayerId + 1;
     int32 MaxLayerId = NodeLayerId;
     	
 	if (!bIsOverlay)
@@ -708,8 +1023,6 @@ int32 SSchematicGraphPanel::OnPaint(const FPaintArgs& Args, const FGeometry& All
 		MaxLayerId++;
 	}
 	
-	
-
 	if (!GraphData)
 	{
 		return MaxLayerId; 
@@ -718,6 +1031,110 @@ int32 SSchematicGraphPanel::OnPaint(const FPaintArgs& Args, const FGeometry& All
 	FArrangedChildren ArrangedChildren(EVisibility::Visible);
 	ArrangeChildNodes(AllottedGeometry, ArrangedChildren);
 
+	NodeCenterByGuid.Reset();
+	NodeCenterByGuid.Reserve(ArrangedChildren.Num());
+	NodeVisibilityByIndex.Reset();
+	NodeVisibilityByIndex.Reserve(ArrangedChildren.Num());
+	NodeVisibilityByGuid.Reset();
+	NodeVisibilityByGuid.Reserve(ArrangedChildren.Num());
+	
+	for (int32 ChildIndex = 0; ChildIndex < ArrangedChildren.Num(); ++ChildIndex)
+	{
+		FArrangedWidget& CurWidget = ArrangedChildren[ChildIndex];
+		const TSharedRef<SSchematicGraphNode> ChildNode = StaticCastSharedRef<SSchematicGraphNode>(CurWidget.Widget);
+
+		NodeCenterByGuid.Add(ChildNode->GetGuid(), CurWidget.Geometry.GetLocalPositionAtCoordinates({0.5, 0.5}));
+
+		const FSchematicGraphNode* NodeData = ChildNode->GetNodeData();
+
+		const bool bIsVisible =
+			!CurWidget.Geometry.GetLocalSize().IsNearlyZero() &&
+			FSlateRect::DoRectanglesIntersect( CurWidget.Geometry.GetLayoutBoundingRect(), MyCullingRect ) &&
+			(GraphData->GetVisibilityForNode(NodeData) != ESchematicGraphVisibility::Hidden);
+
+		NodeVisibilityByIndex.Add(bIsVisible);
+		NodeVisibilityByGuid.Add(NodeData->GetGuid(), bIsVisible);
+	}
+
+	// draw all of the links
+	for(const TPair<FGuid, TSharedPtr<FSchematicLinkWidgetInfo>>& Pair : LinkByGuid)
+	{
+		if(GetVisibilityForLink(Pair.Key) == ESchematicGraphVisibility::Hidden)
+		{
+			continue;
+		}
+
+		const FSchematicGraphLink* Link = GraphData->FindLink(Pair.Key);
+		if(Link == nullptr)
+		{
+			continue;
+		}
+
+		// the node may have been culled
+		const bool* bIsSourceNodeVisible = NodeVisibilityByGuid.Find(Link->GetSourceNodeGuid()); 
+		const bool* bIsTargetNodeVisible = NodeVisibilityByGuid.Find(Link->GetTargetNodeGuid());
+		if((bIsSourceNodeVisible == nullptr) || (bIsTargetNodeVisible == nullptr))
+		{
+			continue;
+		}
+		if(!(*bIsSourceNodeVisible) || !(*bIsTargetNodeVisible))
+		{
+			continue;
+		}
+		
+		const SSchematicGraphNode* SourceNode = FindNode(Link->GetSourceNodeGuid()); 
+		const SSchematicGraphNode* TargetNode = FindNode(Link->GetTargetNodeGuid());
+		if(SourceNode == nullptr || TargetNode == nullptr || SourceNode == TargetNode)
+		{
+			continue;
+		}
+
+		const FLinearColor Color = Pair.Value->Color->Get();
+		const float Thickness = Pair.Value->Thickness->Get();
+		const FSlateBrush* Brush = GetBrushForLink(Pair.Key);
+		const FVector2d& SourcePosition = NodeCenterByGuid.FindChecked(SourceNode->GetGuid()) + GraphData->GetSourceNodeOffsetForLink(Link);
+		const FVector2d& TargetPosition = NodeCenterByGuid.FindChecked(TargetNode->GetGuid()) + GraphData->GetTargetNodeOffsetForLink(Link);
+
+		const FVector2d Diff = TargetPosition - SourcePosition;
+		const float DiffLength = Diff.Size();
+		if(DiffLength < SMALL_NUMBER)
+		{
+			continue;
+		}
+		
+		const float Minimum = Pair.Value->Minimum->Get();
+		const float Maximum = Pair.Value->Maximum->Get();
+
+		const float SourceMinimumDistance = GetMinimumLinkDistanceForNode(SourceNode->GetGuid());
+		const float TargetMinimumDistance = GetMinimumLinkDistanceForNode(TargetNode->GetGuid());
+
+		if(DiffLength <= (SourceMinimumDistance + TargetMinimumDistance))
+		{
+			continue;
+		}
+		
+		const FVector2d DiffNormal = Diff / DiffLength;
+		const FVector2d MinimumPosition = SourcePosition + DiffNormal * SourceMinimumDistance;
+		const FVector2d MaximumPosition = TargetPosition - DiffNormal * TargetMinimumDistance;
+
+		const TArray<FVector2D> LinePoints = {
+			FMath::Lerp<FVector2d>(MinimumPosition, MaximumPosition, FMath::Clamp(Minimum, 0, 1)),
+			FMath::Lerp<FVector2d>(MinimumPosition, MaximumPosition, FMath::Clamp(Maximum, 0, 1))
+		};
+
+		// todo: draw the lines using a brush
+		FSlateDrawElement::MakeLines(
+			OutDrawElements,
+			LinkLayerId,
+			AllottedGeometry.ToPaintGeometry(),
+			LinePoints,
+			ESlateDrawEffect::None,
+			Color,
+			true,
+			Thickness
+		);
+	}
+
 	// Because we paint multiple children, we must track the maximum layer id that they produced in case one of our parents
 	// wants to an overlay for all of its contents.
 
@@ -727,19 +1144,18 @@ int32 SSchematicGraphPanel::OnPaint(const FPaintArgs& Args, const FGeometry& All
 	{
 		for (int32 ChildIndex = 0; ChildIndex < ArrangedChildren.Num(); ++ChildIndex)
 		{
+			if(!NodeVisibilityByIndex[ChildIndex])
+			{
+				continue;
+			}
+			
 			FArrangedWidget& CurWidget = ArrangedChildren[ChildIndex];
 			TSharedRef<SSchematicGraphNode> ChildNode = StaticCastSharedRef<SSchematicGraphNode>(CurWidget.Widget);
 			
 			// Examine node to see what layers we should be drawing in
 			const int32 ChildLayerId = NodeLayerId;
-
-			const bool bNodeIsVisible = FSlateRect::DoRectanglesIntersect( CurWidget.Geometry.GetLayoutBoundingRect(), MyCullingRect );
-
-			if (bNodeIsVisible)
-			{
-				const int32 CurWidgetsMaxLayerId = CurWidget.Widget->Paint(NewArgs, CurWidget.Geometry, MyCullingRect, OutDrawElements, ChildLayerId, InWidgetStyle, true );
-				MaxLayerId = FMath::Max( MaxLayerId, CurWidgetsMaxLayerId + 1 );
-			}
+			const int32 CurWidgetsMaxLayerId = CurWidget.Widget->Paint(NewArgs, CurWidget.Geometry, MyCullingRect, OutDrawElements, ChildLayerId, InWidgetStyle, true );
+			MaxLayerId = FMath::Max( MaxLayerId, CurWidgetsMaxLayerId + 1 );
 		}
 	}
 
@@ -965,8 +1381,7 @@ FVector2d SSchematicGraphPanel::GetSizeForNode(FGuid InNodeGuid) const
 	{
 		return GraphData->GetSizeForNode(InNodeGuid);
 	}
-	static const FVector2d DefaultSize = FVector2d(50.0,50.0);
-	return DefaultSize;
+	return SSchematicGraphNode::DefaultNodeSize;
 }
 
 float SSchematicGraphPanel::GetScaleForNode(FGuid InNodeGuid, bool bIncludeScaleOffset) const
@@ -1001,7 +1416,22 @@ bool SSchematicGraphPanel::IsAutoScaleEnabledForNode(FGuid InNodeGuid) const
 	return false;
 }
 
-FLinearColor SSchematicGraphPanel::GetColorForNode(FGuid InNodeGuid) const
+ float SSchematicGraphPanel::GetMinimumLinkDistanceForNode(FGuid InLinkGuid, bool bIncludeScale) const
+ {
+	if(GraphData)
+	{
+		const float MinimumDistance = GraphData->GetMinimumLinkDistanceForNode(InLinkGuid);
+		if(bIncludeScale)
+		{
+			const float Scale = GetScaleForNode(InLinkGuid, true);
+			return MinimumDistance * Scale;
+		}
+		return MinimumDistance;
+	}
+	return 0.f;
+ }
+
+ FLinearColor SSchematicGraphPanel::GetColorForNode(FGuid InNodeGuid) const
 {
 	if(GraphData)
 	{
@@ -1029,13 +1459,13 @@ FText SSchematicGraphPanel::GetToolTipForNode(FGuid InNodeGuid) const
 	return FText();
 }
 
-ESchematicGraphNodeVisibility SSchematicGraphPanel::GetVisibilityForNode(FGuid InNodeGuid) const
+ESchematicGraphVisibility SSchematicGraphPanel::GetVisibilityForNode(FGuid InNodeGuid) const
 {
 	if(GraphData)
 	{
 		return GraphData->GetVisibilityForNode(InNodeGuid);
 	}
-	return ESchematicGraphNodeVisibility::Visible;
+	return ESchematicGraphVisibility::Visible;
 }
 
 bool SSchematicGraphPanel::IsDragSupportedForNode(FGuid InNodeGuid) const
@@ -1047,6 +1477,70 @@ bool SSchematicGraphPanel::IsDragSupportedForNode(FGuid InNodeGuid) const
 	return false;
 }
 
+ float SSchematicGraphPanel::GetMinimumForLink(FGuid InLinkGuid) const
+ {
+	if(GraphData)
+	{
+		return GraphData->GetMinimumForLink(InLinkGuid);
+	}
+	return 1.f;
+ }
+
+ float SSchematicGraphPanel::GetMaximumForLink(FGuid InLinkGuid) const
+ {
+	if(GraphData)
+	{
+		return GraphData->GetMaximumForLink(InLinkGuid);
+	}
+	return 1.f;
+ }
+
+ FLinearColor SSchematicGraphPanel::GetColorForLink(FGuid InLinkGuid) const
+{
+	if(GraphData)
+	{
+		return GraphData->GetColorForLink(InLinkGuid);
+	}
+	return FLinearColor::White;
+}
+
+ float SSchematicGraphPanel::GetThicknessForLink(FGuid InLinkGuid) const
+ {
+	if(GraphData)
+	{
+		return GraphData->GetThicknessForLink(InLinkGuid);
+	}
+	return 1.f;
+ }
+
+ const FSlateBrush* SSchematicGraphPanel::GetBrushForLink(FGuid InLinkGuid) const
+{
+	if(GraphData)
+	{
+		return GraphData->GetBrushForLink(InLinkGuid);
+	}
+	static const FSlateBrush* DefaultBrush = nullptr;
+	return DefaultBrush;
+}
+
+FText SSchematicGraphPanel::GetToolTipForLink(FGuid InLinkGuid) const
+{
+	if(GraphData)
+	{
+		return GraphData->GetToolTipForLink(InLinkGuid);
+	}
+	return FText();
+}
+
+ESchematicGraphVisibility SSchematicGraphPanel::GetVisibilityForLink(FGuid InLinkGuid) const
+{
+	if(GraphData)
+	{
+		return GraphData->GetVisibilityForLink(InLinkGuid);
+	}
+	return ESchematicGraphVisibility::Visible;
+}
+
 void SSchematicGraphPanel::UpdateAutoScalingForNodes()
 {
 	TArray<bool> NodeIsAutoScaling;
diff --git a/Engine/Source/Editor/AnimationEditorWidgets/Public/SchematicGraphPanel/SSchematicGraphPanel.h b/Engine/Source/Editor/AnimationEditorWidgets/Public/SchematicGraphPanel/SSchematicGraphPanel.h
--- a/Engine/Source/Editor/AnimationEditorWidgets/Public/SchematicGraphPanel/SSchematicGraphPanel.h
+++ b/Engine/Source/Editor/AnimationEditorWidgets/Public/SchematicGraphPanel/SSchematicGraphPanel.h
@@ -19,14 +19,14 @@ enum ESchematicGraphNodePlacementConstraint
 	BottomRight
 };
 
-enum ESchematicGraphNodeVisibility
+enum ESchematicGraphVisibility
 {
 	Visible,
 	FadedOut,
 	Hidden
 };
 
-#define SCHEMATICGRAPHNODE_BODY(ClassName, SuperClass) \
+#define SCHEMATICGRAPHELEMENT_BODY(ClassName, SuperClass, BaseClass) \
 inline static const FName& Type = TEXT(#ClassName); \
 virtual const FName& GetType() const override { return ClassName::Type; } \
 virtual bool IsA(const FName& InType) const \
@@ -40,23 +40,29 @@ virtual bool IsA(const FName& InType) const \
 template<typename T> \
 friend const T* Cast(const ClassName* InNode) \
 { \
-	return Cast<T>((const FSchematicGraphNode*) InNode); \
+	return Cast<T>((const BaseClass*) InNode); \
 } \
 template<typename T> \
 friend T* Cast(ClassName* InNode) \
 { \
-	return Cast<T>((FSchematicGraphNode*) InNode); \
+	return Cast<T>((BaseClass*) InNode); \
 } \
 template<typename T> \
 friend const T* CastChecked(const ClassName* InNode) \
 { \
-	return CastChecked<T>((const FSchematicGraphNode*) InNode); \
+	return CastChecked<T>((const BaseClass*) InNode); \
 } \
 template<typename T> \
 friend T* CastChecked(ClassName* InNode) \
 { \
-	return CastChecked<T>((FSchematicGraphNode*) InNode); \
-} \
+	return CastChecked<T>((BaseClass*) InNode); \
+}
+
+#define SCHEMATICGRAPHNODE_BODY(ClassName, SuperClass) \
+SCHEMATICGRAPHELEMENT_BODY(ClassName, SuperClass, FSchematicGraphNode)
+
+#define SCHEMATICGRAPHLINK_BODY(ClassName, SuperClass) \
+SCHEMATICGRAPHELEMENT_BODY(ClassName, SuperClass, FSchematicGraphNode)
 
 class ANIMATIONEDITORWIDGETS_API FSchematicGraphNode
 {
@@ -128,8 +134,8 @@ public:
 	virtual const FText& GetToolTip() const { return ToolTip; }
 	virtual ESchematicGraphNodePlacementConstraint GetPlacement() const { return Placement; }
 	virtual void SetPlacement(ESchematicGraphNodePlacementConstraint InPlacement) { Placement = InPlacement; }
-	virtual ESchematicGraphNodeVisibility GetVisibility() const { return Visibility; }
-	virtual void SetVisibility(ESchematicGraphNodeVisibility InVisibility) { Visibility = InVisibility; }
+	virtual ESchematicGraphVisibility GetVisibility() const { return Visibility; }
+	virtual void SetVisibility(ESchematicGraphVisibility InVisibility) { Visibility = InVisibility; }
 	virtual bool IsDragSupported() const { return bDragSupported; }
 	virtual void SetDragSupported(bool InDragSupported) { bDragSupported = InDragSupported;}
 
@@ -146,14 +152,109 @@ protected:
 	const FSlateBrush* Brush = nullptr;
 	FText ToolTip = FText();
 	ESchematicGraphNodePlacementConstraint Placement = ESchematicGraphNodePlacementConstraint::Free;
-	ESchematicGraphNodeVisibility Visibility = ESchematicGraphNodeVisibility::Visible;
+	ESchematicGraphVisibility Visibility = ESchematicGraphVisibility::Visible;
 	bool bDragSupported = false;
 
 	friend class FSchematicGraphModel;
 };
 
+class ANIMATIONEDITORWIDGETS_API FSchematicGraphLink
+{
+public:
+	virtual ~FSchematicGraphLink() {}
+	
+	inline static const FName& Type = TEXT("FSchematicGraphLink");
+	virtual const FName& GetType() const;
+	virtual bool IsA(const FName& InType) const;
+	template<typename NodeType>
+	bool IsA() const
+	{
+		return IsA(NodeType::Type);
+	}
+	
+	template<typename T>
+	friend const T* Cast(const FSchematicGraphLink* InNode)
+	{
+		if(InNode)
+		{
+			if(InNode->IsA<T>())
+			{
+				return static_cast<const T*>(InNode);
+			}
+		}
+		return nullptr;
+	}
+
+	template<typename T>
+	friend T* Cast(FSchematicGraphLink* InNode)
+	{
+		if(InNode)
+		{
+			if(InNode->IsA<T>())
+			{
+				return static_cast<T*>(InNode);
+			}
+		}
+		return nullptr;
+	}
+
+	template<typename T>
+	friend const T* CastChecked(const FSchematicGraphLink* InNode)
+	{
+		const T* Node = Cast<T>(InNode);
+		check(Node);
+		return Node;
+	}
+
+	template<typename T>
+	friend T* CastChecked(FSchematicGraphLink* InNode)
+	{
+		T* Node = Cast<T>(InNode);
+		check(Node);
+		return Node;
+	}
+	
+	const FGuid& GetGuid() const { return Guid; }
+	static uint32 GetLinkHash(const FGuid& InSourceNodeGuid, const FGuid& InTargetNodeGuid)
+	{
+		return HashCombine(GetTypeHash(InSourceNodeGuid), GetTypeHash(InTargetNodeGuid));
+	}
+	uint32 GetLinkHash() const { return GetLinkHash(SourceNodeGuid, TargetNodeGuid); }
+	const FGuid& GetSourceNodeGuid() const { return SourceNodeGuid; }
+	const FGuid& GetTargetNodeGuid() const { return TargetNodeGuid; }
+	virtual float GetMinimum() const { return Minimum; }
+	virtual float GetMaximum() const { return Maximum; }
+	virtual FVector2d GetSourceNodeOffset() const { return SourceNodeOffset; }
+	virtual FVector2d GetTargetNodeOffset() const { return TargetNodeOffset; }
+	virtual FLinearColor GetColor() const { return Color; }
+	virtual float GetThickness() const { return Thickness; }
+	virtual const FSlateBrush* GetBrush() const { return Brush; }
+	virtual const FText& GetToolTip() const { return ToolTip; }
+	virtual ESchematicGraphVisibility GetVisibility() const { return Visibility; }
+	virtual void SetVisibility(ESchematicGraphVisibility InVisibility) { Visibility = InVisibility; }
+
+protected:
+	
+	FGuid Guid = FGuid::NewGuid();
+	FGuid SourceNodeGuid = FGuid();
+	FGuid TargetNodeGuid = FGuid();
+	float Minimum = 0.f;
+	float Maximum = 1.f;
+	FVector2d SourceNodeOffset = FVector2d::ZeroVector;
+	FVector2d TargetNodeOffset = FVector2d::ZeroVector;
+	FLinearColor Color = FLinearColor::White;
+	float Thickness = 1.f;
+	const FSlateBrush* Brush = nullptr;
+	FText ToolTip = FText();
+	ESchematicGraphVisibility Visibility = ESchematicGraphVisibility::Visible;
+
+	friend class FSchematicGraphModel;
+};
+
 DECLARE_EVENT_OneParam(FSchematicGraph, FOnNodeAdded, const FSchematicGraphNode*);
 DECLARE_EVENT_OneParam(FSchematicGraph, FOnNodeRemoved, const FSchematicGraphNode*);
+DECLARE_EVENT_OneParam(FSchematicGraph, FOnLinkAdded, const FSchematicGraphLink*);
+DECLARE_EVENT_OneParam(FSchematicGraph, FOnLinkRemoved, const FSchematicGraphLink*);
 DECLARE_EVENT(FSchematicGraph, FOnGraphReset);
 
 class ANIMATIONEDITORWIDGETS_API FSchematicGraphModel
@@ -179,9 +280,9 @@ public:
 	}
 
 	template<typename NodeType = FSchematicGraphNode>
-	const NodeType* FindNode(const FGuid& InGuid) const
+	const NodeType* FindNode(const FGuid& InNodeGuid) const
 	{
-		if(const TSharedPtr<FSchematicGraphNode>* ExistingNode = NodeByGuid.Find(InGuid))
+		if(const TSharedPtr<FSchematicGraphNode>* ExistingNode = NodeByGuid.Find(InNodeGuid))
 		{
 			return Cast<NodeType>(ExistingNode->Get());
 		};
@@ -189,9 +290,9 @@ public:
 	}
 
 	template<typename NodeType = FSchematicGraphNode>
-	const NodeType* FindNodeChecked(const FGuid& InGuid) const
+	const NodeType* FindNodeChecked(const FGuid& InNodeGuid) const
 	{
-		if(const TSharedPtr<FSchematicGraphNode>* ExistingNode = NodeByGuid.Find(InGuid))
+		if(const TSharedPtr<FSchematicGraphNode>* ExistingNode = NodeByGuid.Find(InNodeGuid))
 		{
 			check(ExistingNode->Get()->IsA(NodeType::Type));
 			return static_cast<NodeType*>(ExistingNode->Get());
@@ -199,13 +300,13 @@ public:
 		return nullptr;
 	}
 
-	virtual bool RemoveNode(const FGuid& InGuid);
+	virtual bool RemoveNode(const FGuid& InNodeGuid);
 
 	const TArray<TSharedPtr<FSchematicGraphNode>>& GetNodes() const { return Nodes; }
 
-	FVector2d GetPositionForNode(const FGuid& InGuid) const;
+	FVector2d GetPositionForNode(const FGuid& InNodeGuid) const;
 	virtual FVector2d GetPositionForNode(const FSchematicGraphNode* InNode) const;
-	FVector2d GetPositionOffsetForNode(const FGuid& InGuid) const;
+	FVector2d GetPositionOffsetForNode(const FGuid& InNodeGuid) const;
 	virtual FVector2d GetPositionOffsetForNode(const FSchematicGraphNode* InNode) const;
 	bool GetPositionAnimationEnabledForNode(const FGuid& InNodeGuid) const;
 	virtual bool GetPositionAnimationEnabledForNode(const FSchematicGraphNode* InNode) const;
@@ -215,23 +316,118 @@ public:
 	virtual float GetScaleForNode(const FSchematicGraphNode* InNode, bool bIncludeScaleOffset) const;
 	float GetScaleOffsetForNode(const FGuid& InNodeGuid) const;
 	virtual float GetScaleOffsetForNode(const FSchematicGraphNode* InNode) const;
+	float GetMinimumLinkDistanceForNode(const FGuid& InNodeGuid) const;
+	virtual float GetMinimumLinkDistanceForNode(const FSchematicGraphNode* InNode) const;
 	bool IsAutoScaleEnabledForNode(const FGuid& InNodeGuid) const;
 	virtual bool IsAutoScaleEnabledForNode(const FSchematicGraphNode* InNode) const;
-	FLinearColor GetColorForNode(const FGuid& InGuid) const;
+	FLinearColor GetColorForNode(const FGuid& InNodeGuid) const;
 	virtual FLinearColor GetColorForNode(const FSchematicGraphNode* InNode) const;
-	const FSlateBrush* GetBrushForNode(const FGuid& InGuid) const;
+	const FSlateBrush* GetBrushForNode(const FGuid& InNodeGuid) const;
 	virtual const FSlateBrush* GetBrushForNode(const FSchematicGraphNode* InNode) const;
-	const FText GetToolTipForNode(const FGuid& InGuid) const;
+	const FText GetToolTipForNode(const FGuid& InNodeGuid) const;
 	virtual const FText GetToolTipForNode(const FSchematicGraphNode* InNode) const;
-	ESchematicGraphNodePlacementConstraint GetPlacementForNode(const FGuid& InGuid) const;
+	ESchematicGraphNodePlacementConstraint GetPlacementForNode(const FGuid& InNodeGuid) const;
 	virtual ESchematicGraphNodePlacementConstraint GetPlacementForNode(const FSchematicGraphNode* InNode) const;
-	ESchematicGraphNodeVisibility GetVisibilityForNode(const FGuid& InGuid) const;
-	virtual ESchematicGraphNodeVisibility GetVisibilityForNode(const FSchematicGraphNode* InNode) const;
-	bool IsDragSupportedForNode(const FGuid& InGuid) const;
+	ESchematicGraphVisibility GetVisibilityForNode(const FGuid& InNodeGuid) const;
+	virtual ESchematicGraphVisibility GetVisibilityForNode(const FSchematicGraphNode* InNode) const;
+	bool IsDragSupportedForNode(const FGuid& InNodeGuid) const;
 	virtual bool IsDragSupportedForNode(const FSchematicGraphNode* InNode) const;
 
+	template<typename LinkType = FSchematicGraphLink>
+	LinkType* AddLink(const FGuid& InSourceNodeGuid, const FGuid& InTargetNodeGuid, bool bNotify = true)
+	{
+		const TSharedPtr<FSchematicGraphLink> NewLink = MakeShareable(new LinkType);
+		NewLink->SourceNodeGuid = InSourceNodeGuid;
+		NewLink->TargetNodeGuid = InTargetNodeGuid;
+		Links.Add(NewLink);
+		LinkByGuid.Add(NewLink->GetGuid(), NewLink);
+		LinkByHash.Add(NewLink->GetLinkHash(), NewLink);
+		NodeGuidToLinkGuids.FindOrAdd(InSourceNodeGuid).Get<0>().Add(NewLink->GetGuid());
+		NodeGuidToLinkGuids.FindOrAdd(InTargetNodeGuid).Get<1>().Add(NewLink->GetGuid());
+
+		if (bNotify && OnLinkAddedDelegate.IsBound())
+		{
+			OnLinkAddedDelegate.Broadcast(NewLink.Get());
+		}
+		return static_cast<LinkType*>(NewLink.Get());
+	}
+
+	template<typename LinkType = FSchematicGraphLink>
+	const LinkType* FindLink(const FGuid& InLinkGuid) const
+	{
+		if(const TSharedPtr<FSchematicGraphLink>* ExistingLink = LinkByGuid.Find(InLinkGuid))
+		{
+			return Cast<LinkType>(ExistingLink->Get());
+		};
+		return nullptr;
+	}
+
+	template<typename LinkType = FSchematicGraphLink>
+	const LinkType* FindLinkChecked(const FGuid& InLinkGuid) const
+	{
+		if(const TSharedPtr<FSchematicGraphLink>* ExistingLink = LinkByGuid.Find(InLinkGuid))
+		{
+			check(ExistingLink->Get()->IsA(LinkType::Type));
+			return static_cast<LinkType*>(ExistingLink->Get());
+		};
+		return nullptr;
+	}
+
+	template<typename LinkType = FSchematicGraphLink>
+	const LinkType* FindLink(const FGuid& InSourceNodeGuid, const FGuid& InTargetNodeGuid) const
+	{
+		const uint32 LinkHash = FSchematicGraphLink::GetLinkHash(InSourceNodeGuid, InTargetNodeGuid);
+		if(const TSharedPtr<FSchematicGraphLink>* ExistingLink = LinkByHash.Find(LinkHash))
+		{
+			return Cast<LinkType>(ExistingLink->Get());
+		};
+		return nullptr;
+	}
+
+	template<typename LinkType = FSchematicGraphLink>
+	const LinkType* FindLinkChecked(const FGuid& InSourceNodeGuid, const FGuid& InTargetNodeGuid) const
+	{
+		const uint32 LinkHash = FSchematicGraphLink::GetLinkHash(InSourceNodeGuid, InTargetNodeGuid);
+		if(const TSharedPtr<FSchematicGraphLink>* ExistingLink = LinkByHash.Find(LinkHash))
+		{
+			check(ExistingLink->Get()->IsA(LinkType::Type));
+			return Cast<LinkType>(ExistingLink->Get());
+		};
+		return nullptr;
+	}
+
+	bool IsLinkedTo(const FGuid& InSourceNodeGuid, const FGuid& InTargetNodeGuid) const;
+	TArray<const FSchematicGraphLink*> FindLinksOnNode(const FGuid& InNodeGuid) const;
+	TArray<const FSchematicGraphLink*> FindLinksOnSource(const FGuid& InSourceNodeGuid) const;
+	TArray<const FSchematicGraphLink*> FindLinksOnTarget(const FGuid& InTargetNodeGuid) const;
+
+	virtual bool RemoveLink(const FGuid& InLinkGuid);
+
+	const TArray<TSharedPtr<FSchematicGraphLink>>& GetLinks() const { return Links; }
+
+	float GetMinimumForLink(const FGuid& InLinkGuid) const;
+	virtual float GetMinimumForLink(const FSchematicGraphLink* InLink) const;
+	float GetMaximumForLink(const FGuid& InLinkGuid) const;
+	virtual float GetMaximumForLink(const FSchematicGraphLink* InLink) const;
+	FVector2d GetSourceNodeOffsetForLink(const FGuid& InLinkGuid) const;
+	FVector2d GetSourceNodeOffsetForLink(const FSchematicGraphLink* InLink) const;
+	FVector2d GetTargetNodeOffsetForLink(const FGuid& InLinkGuid) const;
+	FVector2d GetTargetNodeOffsetForLink(const FSchematicGraphLink* InLink) const;
+	FLinearColor GetColorForLink(const FGuid& InLinkGuid) const;
+	virtual FLinearColor GetColorForLink(const FSchematicGraphLink* InLink) const;
+	float GetThicknessForLink(const FGuid& InLinkGuid) const;
+	virtual float GetThicknessForLink(const FSchematicGraphLink* InLink) const;
+	const FSlateBrush* GetBrushForLink(const FGuid& InLinkGuid) const;
+	virtual const FSlateBrush* GetBrushForLink(const FSchematicGraphLink* InLink) const;
+	const FText GetToolTipForLink(const FGuid& InLinkGuid) const;
+	virtual const FText GetToolTipForLink(const FSchematicGraphLink* InLink) const;
+	ESchematicGraphVisibility GetVisibilityForLink(const FGuid& InLinkGuid) const;
+	virtual ESchematicGraphVisibility GetVisibilityForLink(const FSchematicGraphLink* InLink) const;
+
 	FOnNodeAdded& OnNodeAdded() { return OnNodeAddedDelegate; }
 	FOnNodeRemoved& OnNodeRemoved() { return OnNodeRemovedDelegate; }
+	FOnLinkAdded& OnLinkAdded() { return OnLinkAddedDelegate; }
+	FOnLinkRemoved& OnLinkRemoved() { return OnLinkRemovedDelegate; }
 	FOnGraphReset& OnGraphReset() { return OnGraphResetDelegate; }
 
 	virtual bool GetForwardedNodeForDrag(FGuid& InOutGuid) const { return false; }
@@ -240,10 +436,15 @@ protected:
 
 	TArray<TSharedPtr<FSchematicGraphNode>> Nodes;
 	TMap<FGuid, TSharedPtr<FSchematicGraphNode>> NodeByGuid;
-	TMap<FGuid, TArray<FGuid>> Links;
+	TArray<TSharedPtr<FSchematicGraphLink>> Links;
+	TMap<FGuid, TSharedPtr<FSchematicGraphLink>> LinkByGuid;
+	TMap<uint32, TSharedPtr<FSchematicGraphLink>> LinkByHash;
+	TMap<FGuid, TTuple<TArray<FGuid>, TArray<FGuid>>> NodeGuidToLinkGuids;
 
 	FOnNodeAdded OnNodeAddedDelegate;
 	FOnNodeRemoved OnNodeRemovedDelegate;
+	FOnLinkAdded OnLinkAddedDelegate;
+	FOnLinkRemoved OnLinkRemovedDelegate;
 	FOnGraphReset OnGraphResetDelegate;
 };
 
@@ -344,7 +545,8 @@ public:
 private:
 
 	bool bIsBeingDragged = false;
-	FVector2d OriginalSize = FVector2d(50.0,50.0);
+	static inline const FVector2d DefaultNodeSize = FVector2d(32.0,32.0);  
+	FVector2d OriginalSize = DefaultNodeSize;
 	static inline constexpr float ScaledUp = 1.25;
 	static inline constexpr float ScaledDown = 0.75;
 
@@ -365,6 +567,7 @@ private:
 	SSchematicGraphPanel* SchematicGraphPanel = nullptr;
 
 	friend class SSchematicGraphPanel;
+	friend class FSchematicGraphModel;
 };
 
 
@@ -396,6 +599,14 @@ public:
 	SLATE_EVENT(FOnDrop, OnDrop)
 	SLATE_END_ARGS()
 
+	struct FSchematicLinkWidgetInfo
+	{
+		TSharedPtr<FFloatAttribute> Minimum;
+		TSharedPtr<FFloatAttribute> Maximum;
+		TSharedPtr<FLinearColorAttribute> Color;
+		TSharedPtr<FFloatAttribute> Thickness;
+	};
+
 	virtual ~SSchematicGraphPanel() override
 	{
 		OnNodeClickedDelegate.Unbind();
@@ -409,6 +620,9 @@ public:
 	void AddNode(const FSchematicGraphNode* InNodeToAdd);
 	void RemoveNode(const FSchematicGraphNode* InNodeToRemove);
 	const SSchematicGraphNode* FindNode(const FGuid& InGuid) const;
+	void AddLink(const FSchematicGraphLink* InLinkToAdd);
+	void RemoveLink(const FSchematicGraphLink* InLinkToRemove);
+	const FSchematicLinkWidgetInfo* FindLink(const FGuid& InGuid) const;
 	
 	// SNodePanel interface
 	virtual void OnArrangeChildren(const FGeometry& AllottedGeometry, FArrangedChildren& ArrangedChildren) const override;
@@ -436,12 +650,21 @@ public:
 	virtual FVector2d GetSizeForNode(FGuid InNodeGuid) const;
 	virtual float GetScaleForNode(FGuid InNodeGuid, bool bIncludeScaleOffset) const;
 	virtual bool IsAutoScaleEnabledForNode(FGuid InNodeGuid) const;
+	virtual float GetMinimumLinkDistanceForNode(FGuid InLinkGuid, bool bIncludeScale = true) const;
 	virtual FLinearColor GetColorForNode(FGuid InNodeGuid) const;
 	virtual const FSlateBrush* GetBrushForNode(FGuid InNodeGuid) const;
 	virtual FText GetToolTipForNode(FGuid InNodeGuid) const;
-	virtual ESchematicGraphNodeVisibility GetVisibilityForNode(FGuid InNodeGuid) const;
+	virtual ESchematicGraphVisibility GetVisibilityForNode(FGuid InNodeGuid) const;
 	virtual bool IsDragSupportedForNode(FGuid InNodeGuid) const;
 
+	virtual float GetMinimumForLink(FGuid InLinkGuid) const;
+	virtual float GetMaximumForLink(FGuid InLinkGuid) const;
+	virtual FLinearColor GetColorForLink(FGuid InLinkGuid) const;
+	virtual float GetThicknessForLink(FGuid InLinkGuid) const;
+	virtual const FSlateBrush* GetBrushForLink(FGuid InLinkGuid) const;
+	virtual FText GetToolTipForLink(FGuid InLinkGuid) const;
+	virtual ESchematicGraphVisibility GetVisibilityForLink(FGuid InLinkGuid) const;
+
 	void UpdateAutoScalingForNodes();
 
 	bool bIsDragDropping = false;
@@ -462,6 +685,12 @@ public:
 	TArray<FGuid> NodesBottomLeft;
 	TArray<FGuid> NodesBottomRight;
 	TMap<FGuid, TSharedPtr<SSchematicGraphNode>> NodeByGuid;
+
+	TMap<FGuid, TSharedPtr<FSchematicLinkWidgetInfo>> LinkByGuid;
+
+	mutable TMap<FGuid, FVector2d> NodeCenterByGuid;
+	mutable TArray<bool> NodeVisibilityByIndex;
+	mutable TMap<FGuid, bool> NodeVisibilityByGuid;
 };
 
 #endif
\ No newline at end of file
diff --git a/Engine/Source/Runtime/SlateCore/Public/Animation/AttributeInterpolator.h b/Engine/Source/Runtime/SlateCore/Public/Animation/AttributeInterpolator.h
--- a/Engine/Source/Runtime/SlateCore/Public/Animation/AttributeInterpolator.h
+++ b/Engine/Source/Runtime/SlateCore/Public/Animation/AttributeInterpolator.h
@@ -131,7 +131,10 @@ public:
 				return;
 			}
 
-			LastValue = InterpolatedValue.IsSet() ? InterpolatedValue : DesiredValue;
+			if(!LastValue.IsSet())
+			{
+				LastValue = InterpolatedValue.IsSet() ? InterpolatedValue : DesiredValue;
+			}
 			DesiredValue = InValue;
 		}
 		else
