Control Rig: Schematic Improvements

Multi layer node rendering
SVG based link rendering

#rnx

diff --git a/Engine/Plugins/Animation/ControlRig/Source/ControlRig/Private/Rigs/RigHierarchy.cpp b/Engine/Plugins/Animation/ControlRig/Source/ControlRig/Private/Rigs/RigHierarchy.cpp
--- a/Engine/Plugins/Animation/ControlRig/Source/ControlRig/Private/Rigs/RigHierarchy.cpp
+++ b/Engine/Plugins/Animation/ControlRig/Source/ControlRig/Private/Rigs/RigHierarchy.cpp
@@ -5405,6 +5405,54 @@ void URigHierarchy::FMetadataStorage::Serialize(FArchive& Ar)
 	}
 }
 
+void URigHierarchy::PropagateMetadata(const FRigElementKey& InKey, const FName& InName, bool bNotify)
+{
+	if(const FRigBaseElement* Element = Find(InKey))
+	{
+		PropagateMetadata(Element, InName, bNotify);
+	}
+}
+
+void URigHierarchy::PropagateMetadata(const FRigBaseElement* InElement, const FName& InName, bool bNotify)
+{
+#if WITH_EDITOR
+	check(InElement);
+
+	if (!ElementMetadata.IsValidIndex(InElement->MetadataStorageIndex))
+	{
+		return;
+	}
+
+	FMetadataStorage& Storage = ElementMetadata[InElement->MetadataStorageIndex];
+	FRigBaseMetadata** MetadataPtrPtr = Storage.MetadataMap.Find(InName);
+	if (!MetadataPtrPtr)
+	{
+		return;
+	}
+
+	FRigBaseMetadata* MetadataPtr = *MetadataPtrPtr;
+	if(MetadataPtr == nullptr)
+	{
+		return;
+	}
+
+	ForEachListeningHierarchy([InElement, MetadataPtr, bNotify](const FRigHierarchyListener& Listener)
+	{
+		if(URigHierarchy* ListeningHierarchy = Listener.Hierarchy.Get())
+		{
+			if(FRigBaseElement* Element = ListeningHierarchy->Find(InElement->GetKey()))
+			{
+				if (FRigBaseMetadata* Metadata = ListeningHierarchy->GetMetadataForElement(Element, MetadataPtr->GetName(), MetadataPtr->GetType(), bNotify))
+				{
+					Metadata->SetValueData(MetadataPtr->GetValueData(), MetadataPtr->GetValueSize());
+					ListeningHierarchy->PropagateMetadata(Element, Metadata->GetName(), bNotify);
+				}
+			}
+		}
+	});
+#endif
+}
+
 void URigHierarchy::OnMetadataChanged(const FRigElementKey& InKey, const FName& InName)
 {
 	MetadataVersion++;
diff --git a/Engine/Plugins/Animation/ControlRig/Source/ControlRig/Private/Rigs/RigHierarchyElements.cpp b/Engine/Plugins/Animation/ControlRig/Source/ControlRig/Private/Rigs/RigHierarchyElements.cpp
--- a/Engine/Plugins/Animation/ControlRig/Source/ControlRig/Private/Rigs/RigHierarchyElements.cpp
+++ b/Engine/Plugins/Animation/ControlRig/Source/ControlRig/Private/Rigs/RigHierarchyElements.cpp
@@ -1607,6 +1607,7 @@ void FRigSocketElement::SetColor(const FLinearColor& InColor, URigHierarchy* InH
 		return;
 	}
 	InHierarchy->SetLinearColorMetadata(GetKey(), ColorMetaName, InColor);
+	InHierarchy->PropagateMetadata(GetKey(), ColorMetaName, bNotify);
 	if(bNotify)
 	{
 		InHierarchy->Notify(ERigHierarchyNotification::SocketColorChanged, this);
@@ -1631,6 +1632,7 @@ void FRigSocketElement::SetDescription(const FString& InDescription, URigHierarc
 		return;
 	}
 	InHierarchy->SetNameMetadata(GetKey(), DescriptionMetaName, *InDescription);
+	InHierarchy->PropagateMetadata(this, DescriptionMetaName, bNotify);
 	if(bNotify)
 	{
 		InHierarchy->Notify(ERigHierarchyNotification::SocketDescriptionChanged, this);
diff --git a/Engine/Plugins/Animation/ControlRig/Source/ControlRig/Public/Rigs/RigHierarchy.h b/Engine/Plugins/Animation/ControlRig/Source/ControlRig/Public/Rigs/RigHierarchy.h
--- a/Engine/Plugins/Animation/ControlRig/Source/ControlRig/Public/Rigs/RigHierarchy.h
+++ b/Engine/Plugins/Animation/ControlRig/Source/ControlRig/Public/Rigs/RigHierarchy.h
@@ -4625,6 +4625,13 @@ private:
 		return SetMetadata<TArray<T>>(InElement, InType, InMetadataName, InValue);
 	}
 
+public:
+	
+	void PropagateMetadata(const FRigElementKey& InKey, const FName& InName, bool bNotify = true);
+	void PropagateMetadata(const FRigBaseElement* InElement, const FName& InName, bool bNotify = true);
+	
+private:
+	
 	void OnMetadataChanged(const FRigElementKey& InKey, const FName& InName);
 	void OnMetadataTagChanged(const FRigElementKey& InKey, const FName& InTag, bool bAdded);
 
diff --git a/Engine/Plugins/Animation/ControlRig/Source/ControlRigEditor/Private/ControlRigEditorStyle.h b/Engine/Plugins/Animation/ControlRig/Source/ControlRigEditor/Private/ControlRigEditorStyle.h
--- a/Engine/Plugins/Animation/ControlRig/Source/ControlRigEditor/Private/ControlRigEditorStyle.h
+++ b/Engine/Plugins/Animation/ControlRig/Source/ControlRigEditor/Private/ControlRigEditorStyle.h
@@ -148,14 +148,17 @@ public:
 
 		// Schematic
 		{
-			Set("ControlRig.Schematic.SocketUnresolved", new IMAGE_BRUSH_SVG("Slate/Socket_Open", Icon128x128));
-			Set("ControlRig.Schematic.SocketResolved", new IMAGE_BRUSH_SVG("Slate/Socket_Closed", Icon128x128));
+			Set("ControlRig.Schematic.Outline", new IMAGE_BRUSH_SVG("Slate/Outline_Schematic", Icon128x128));
+			Set("ControlRig.Schematic.Background", new IMAGE_BRUSH_SVG("Slate/Background_Schematic", Icon128x128));
+			Set("ControlRig.Schematic.SocketUnresolved", new IMAGE_BRUSH_SVG("Slate/Socket_Open_Schematic", Icon128x128));
+			Set("ControlRig.Schematic.SocketResolved", new IMAGE_BRUSH_SVG("Slate/Socket_Closed_Schematic", Icon128x128));
 			Set("ControlRig.Schematic.ConnectorPrimary", new IMAGE_BRUSH_SVG("Slate/Connector_Primary", Icon128x128));
 			Set("ControlRig.Schematic.ConnectorSecondary", new IMAGE_BRUSH_SVG("Slate/Connector_Secondary", Icon128x128));
 			Set("ControlRig.Schematic.ConnectorOptional", new IMAGE_BRUSH_SVG("Slate/Connector_Optional", Icon128x128));
 			Set("ControlRig.Schematic.Bone", new IMAGE_BRUSH_SVG("Slate/Bone_Schematic", Icon128x128));
 			Set("ControlRig.Schematic.Control", new IMAGE_BRUSH_SVG("Slate/Control_Schematic", Icon128x128));
 			Set("ControlRig.Schematic.Null", new IMAGE_BRUSH_SVG("Slate/Null_Schematic", Icon128x128));
+			Set("ControlRig.Schematic.Link", new IMAGE_BRUSH_SVG("Slate/Link_Schematic", Icon128x128));
 		}
 
 		FSlateStyleRegistry::RegisterSlateStyle(*this);
diff --git a/Engine/Plugins/Animation/ControlRig/Source/ControlRigEditor/Private/ControlRigElementDetails.cpp b/Engine/Plugins/Animation/ControlRig/Source/ControlRigEditor/Private/ControlRigElementDetails.cpp
--- a/Engine/Plugins/Animation/ControlRig/Source/ControlRigEditor/Private/ControlRigElementDetails.cpp
+++ b/Engine/Plugins/Animation/ControlRig/Source/ControlRigEditor/Private/ControlRigElementDetails.cpp
@@ -5307,7 +5307,6 @@ void FRigSocketElementDetails::CustomizeSettings(IDetailLayoutBuilder& DetailBui
 
 	IDetailCategoryBuilder& SettingsCategory = DetailBuilder.EditCategory(TEXT("Settings"), LOCTEXT("Settings", "Settings"));
 
-	/*
 	SettingsCategory.AddCustomRow(FText::FromString(TEXT("Color")))
 	.NameContent()
 	[
@@ -5322,7 +5321,6 @@ void FRigSocketElementDetails::CustomizeSettings(IDetailLayoutBuilder& DetailBui
 		.Color(this, &FRigSocketElementDetails::GetSocketColor) 
 		.OnMouseButtonDown(this, &FRigSocketElementDetails::SetSocketColor)
 	];
-	*/
 
 	SettingsCategory.AddCustomRow(FText::FromString(TEXT("Description")))
 	.NameContent()
diff --git a/Engine/Plugins/Animation/ControlRig/Source/ControlRigEditor/Private/Editor/ControlRigEditor.cpp b/Engine/Plugins/Animation/ControlRig/Source/ControlRigEditor/Private/Editor/ControlRigEditor.cpp
--- a/Engine/Plugins/Animation/ControlRig/Source/ControlRigEditor/Private/Editor/ControlRigEditor.cpp
+++ b/Engine/Plugins/Animation/ControlRig/Source/ControlRigEditor/Private/Editor/ControlRigEditor.cpp
@@ -2577,7 +2577,7 @@ void FControlRigEditor::HandleSchematicViewportCreated(const TSharedRef<SSchemat
 	InViewport->OnDropDelegate.BindRaw(&SchematicModel, &FControlRigSchematicModel::HandleSchematicDrop);
 }
 
-FVector2D FControlRigEditor::ComputePersonaProjectedScreenPos(const FVector& InWorldPos)
+FVector2D FControlRigEditor::ComputePersonaProjectedScreenPos(const FVector& InWorldPos, bool bClampToScreenRectangle)
 {
 	if (PreviewViewport.IsValid())
 	{
@@ -2612,8 +2612,11 @@ FVector2D FControlRigEditor::ComputePersonaProjectedScreenPos(const FVector& InW
 		ScreenPos = FIntPoint(FMath::FloorToInt(ScreenPos.X), FMath::FloorToInt(ScreenPos.Y));
 
 		// Clamp to screen rect
-		ScreenPos.X = FMath::Clamp(ScreenPos.X, ViewRect.Min.X, ViewRect.Max.X);
-		ScreenPos.Y = FMath::Clamp(ScreenPos.Y, ViewRect.Min.Y, ViewRect.Max.Y);
+		if(bClampToScreenRectangle)
+		{
+			ScreenPos.X = FMath::Clamp(ScreenPos.X, ViewRect.Min.X, ViewRect.Max.X);
+			ScreenPos.Y = FMath::Clamp(ScreenPos.Y, ViewRect.Min.Y, ViewRect.Max.Y);
+		}
 
 		return FVector2D(ScreenPos.X, ScreenPos.Y);
 	}
diff --git a/Engine/Plugins/Animation/ControlRig/Source/ControlRigEditor/Private/Editor/ControlRigEditor.h b/Engine/Plugins/Animation/ControlRig/Source/ControlRigEditor/Private/Editor/ControlRigEditor.h
--- a/Engine/Plugins/Animation/ControlRig/Source/ControlRigEditor/Private/Editor/ControlRigEditor.h
+++ b/Engine/Plugins/Animation/ControlRig/Source/ControlRigEditor/Private/Editor/ControlRigEditor.h
@@ -193,7 +193,7 @@ public:
 	FReply OnNavigateToConnectorWarning() const;
 	FSimpleMulticastDelegate& OnRequestNavigateToConnectorWarning() { return RequestNavigateToConnectorWarningDelegate; }
 
-	FVector2D ComputePersonaProjectedScreenPos(const FVector& InWorldPos);
+	FVector2D ComputePersonaProjectedScreenPos(const FVector& InWorldPos, bool bClampToScreenRectangle = false);
 	
 protected:
 
diff --git a/Engine/Plugins/Animation/ControlRig/Source/ControlRigEditor/Private/Editor/ControlRigSchematicModel.cpp b/Engine/Plugins/Animation/ControlRig/Source/ControlRigEditor/Private/Editor/ControlRigSchematicModel.cpp
--- a/Engine/Plugins/Animation/ControlRig/Source/ControlRigEditor/Private/Editor/ControlRigSchematicModel.cpp
+++ b/Engine/Plugins/Animation/ControlRig/Source/ControlRigEditor/Private/Editor/ControlRigSchematicModel.cpp
@@ -58,6 +58,15 @@ FControlRigSchematicRigElementKeyNode* FControlRigSchematicModel::AddElementKeyN
 	if(Node)
 	{
 		Node->Key = InKey;
+
+		/*
+		FSchematicGraphTag* Tag = Node->AddTag();
+		static const FSlateBrush* BackgroundBrush = FControlRigEditorStyle::Get().GetBrush( "ControlRig.Schematic.Background");
+		Tag->SetBackgroundBrush(BackgroundBrush);
+		Tag->SetLabelColor(FLinearColor::White);
+		Tag->SetLabel(FText::FromString(FString::FromInt(Nodes.Num() * 7)));
+		*/
+
 		RigElementKeyToGuid.Add(Node->GetKey(), Node->GetGuid());
 
 		if (bNotify && OnNodeAddedDelegate.IsBound())
@@ -126,7 +135,7 @@ FControlRigSchematicRigElementKeyLink* FControlRigSchematicModel::AddElementKeyL
 		FControlRigSchematicRigElementKeyLink* Link = AddLink<FControlRigSchematicRigElementKeyLink>(SourceNode->GetGuid(), TargetNode->GetGuid());
 		Link->SourceKey = InSourceKey;
 		Link->TargetKey = InTargetKey;
-		Link->Thickness = 2.5f;
+		Link->Thickness = 12.f;
 		return Link;
 	}
 	return nullptr;
@@ -469,10 +478,30 @@ bool FControlRigSchematicModel::GetPositionAnimationEnabledForNode(const FSchema
 	return FSchematicGraphModel::GetPositionAnimationEnabledForNode(InNode);
 }
 
-const FSlateBrush* FControlRigSchematicModel::GetBrushForNode(const FSchematicGraphNode* InNode) const
+int32 FControlRigSchematicModel::GetNumLayersForNode(const FSchematicGraphNode* InNode) const
+{
+	if(const FControlRigSchematicRigElementKeyNode* Node = Cast<FControlRigSchematicRigElementKeyNode>(InNode))
+	{
+		return 3;
+	}
+	return FSchematicGraphModel::GetNumLayersForNode(InNode);
+}
+
+const FSlateBrush* FControlRigSchematicModel::GetBrushForNode(const FSchematicGraphNode* InNode, int32 InLayerIndex) const
 {
 	if(const FControlRigSchematicRigElementKeyNode* Node = Cast<FControlRigSchematicRigElementKeyNode>(InNode))
 	{
+		if(InLayerIndex == 0)
+		{
+			static const FSlateBrush* BackgroundBrush = FControlRigEditorStyle::Get().GetBrush( "ControlRig.Schematic.Background");
+			return BackgroundBrush;
+		}
+		if(InLayerIndex == 1)
+		{
+			static const FSlateBrush* OutlineBrush = FControlRigEditorStyle::Get().GetBrush( "ControlRig.Schematic.Outline");
+			return OutlineBrush;
+		}
+		
 		switch(Node->GetKey().Type)
 		{
 			case ERigElementType::Socket:
@@ -538,13 +567,22 @@ const FSlateBrush* FControlRigSchematicModel::GetBrushForNode(const FSchematicGr
 			}
 		}
 	}
-	return FSchematicGraphModel::GetBrushForNode(InNode);
+	return FSchematicGraphModel::GetBrushForNode(InNode, InLayerIndex);
 }
 
-FLinearColor FControlRigSchematicModel::GetColorForNode(const FSchematicGraphNode* InNode) const
+FLinearColor FControlRigSchematicModel::GetColorForNode(const FSchematicGraphNode* InNode, int32 InLayerIndex) const
 {
 	if(const FControlRigSchematicRigElementKeyNode* ElementKeyNode = Cast<FControlRigSchematicRigElementKeyNode>(InNode))
 	{
+		if(InLayerIndex == 0) // background
+		{
+			return FLinearColor(0, 0, 0, 0.75);
+		}
+		if(InLayerIndex == 1) // outline
+		{
+			return FLinearColor::White;
+		}
+		
 		switch(ElementKeyNode->GetKey().Type)
 		{
 			case ERigElementType::Bone:
@@ -570,20 +608,33 @@ FLinearColor FControlRigSchematicModel::GetColorForNode(const FSchematicGraphNod
 			}
 			case ERigElementType::Socket:
 			{
-				return FLinearColor::White;
-				//return FControlRigEditorStyle::Get().SocketUserInterfaceColor;
+				if (ControlRigBeingDebuggedPtr.IsValid())
+				{
+					if (const URigHierarchy* Hierarchy = ControlRigBeingDebuggedPtr->GetHierarchy())
+					{
+						if(const FRigSocketElement* Socket = Hierarchy->Find<FRigSocketElement>(ElementKeyNode->GetKey()))
+						{
+							return Socket->GetColor(Hierarchy);
+						}
+					}
+				}
+				break;
 			}
 			case ERigElementType::Connector:
 			{
 				return FLinearColor::White;
 				//return FControlRigEditorStyle::Get().ConnectorUserInterfaceColor;
 			}
+			default:
+			{
+				break;
+			}
 		}
 	}
-	return FSchematicGraphModel::GetColorForNode(InNode);
+	return FSchematicGraphModel::GetColorForNode(InNode, InLayerIndex);
 }
 
-const FText FControlRigSchematicModel::GetToolTipForNode(const FSchematicGraphNode* InNode) const
+FText FControlRigSchematicModel::GetToolTipForNode(const FSchematicGraphNode* InNode) const
 {
 	if(const FControlRigSchematicRigElementKeyNode* Node = Cast<FControlRigSchematicRigElementKeyNode>(InNode))
 	{
@@ -617,7 +668,7 @@ const FText FControlRigSchematicModel::GetToolTipForNode(const FSchematicGraphNo
 	return FSchematicGraphModel::GetToolTipForNode(InNode);
 }
 
-ESchematicGraphVisibility FControlRigSchematicModel::GetVisibilityForNode(const FSchematicGraphNode* InNode) const
+ESchematicGraphVisibility::Type FControlRigSchematicModel::GetVisibilityForNode(const FSchematicGraphNode* InNode) const
 {
 	if(const FControlRigSchematicRigElementKeyNode* Node = Cast<FControlRigSchematicRigElementKeyNode>(InNode))
 	{
@@ -636,7 +687,7 @@ ESchematicGraphVisibility FControlRigSchematicModel::GetVisibilityForNode(const
 	return FSchematicGraphModel::GetVisibilityForNode(InNode);
 }
 
-ESchematicGraphNodePlacementConstraint FControlRigSchematicModel::GetPlacementForNode(const FSchematicGraphNode* InNode) const
+ESchematicGraphPlacementConstraint::Type FControlRigSchematicModel::GetPlacementForNode(const FSchematicGraphNode* InNode) const
 {
 	if(const FControlRigSchematicRigElementKeyNode* Node = Cast<FControlRigSchematicRigElementKeyNode>(InNode))
 	{
@@ -644,13 +695,47 @@ ESchematicGraphNodePlacementConstraint FControlRigSchematicModel::GetPlacementFo
 		{
 			if(!IsConnectorResolved(Node->GetKey()))
 			{
-				return ESchematicGraphNodePlacementConstraint::BottomRight;
+				return ESchematicGraphPlacementConstraint::BottomRight;
 			}
 		}
 	}
 	return FSchematicGraphModel::GetPlacementForNode(InNode);
 }
 
+const FSlateBrush* FControlRigSchematicModel::GetBrushForLink(const FSchematicGraphLink* InLink) const
+{
+	if(const FSchematicGraphNode* SourceNode = FindNode(InLink->GetSourceNodeGuid()))
+	{
+		if(const FSchematicGraphNode* TargetNode = FindNode(InLink->GetTargetNodeGuid()))
+		{
+			if(SourceNode->IsA<FControlRigSchematicRigElementKeyNode>() &&
+				TargetNode->IsA<FControlRigSchematicRigElementKeyNode>())
+			{
+				static const FSlateBrush* LinkBrush = FControlRigEditorStyle::Get().GetBrush( "ControlRig.Schematic.Link");
+				return LinkBrush;
+			}
+		}
+	}
+	return FSchematicGraphModel::GetBrushForLink(InLink);
+}
+
+FLinearColor FControlRigSchematicModel::GetColorForLink(const FSchematicGraphLink* InLink) const
+{
+	if(const FSchematicGraphNode* SourceNode = FindNode(InLink->GetSourceNodeGuid()))
+	{
+		if(const FSchematicGraphNode* TargetNode = FindNode(InLink->GetTargetNodeGuid()))
+		{
+			if(SourceNode->IsA<FControlRigSchematicRigElementKeyNode>() &&
+				TargetNode->IsA<FControlRigSchematicRigElementKeyNode>())
+			{
+				// semi transparent dark gray
+				return FLinearColor(0.7, 0.7, 0.7, .5);
+			}
+		}
+	}
+	return FSchematicGraphModel::GetColorForLink(InLink);
+}
+
 bool FControlRigSchematicModel::GetForwardedNodeForDrag(FGuid& InOutGuid) const
 {
 	if(const FControlRigSchematicRigElementKeyNode* ElementKeyNode = Cast<FControlRigSchematicRigElementKeyNode>(FindNode(InOutGuid)))
@@ -779,7 +864,7 @@ void FControlRigSchematicModel::HandleSchematicBeginDrag(SSchematicGraphPanel* I
 			ExistingElementKeyNode->SetVisibility(Matches.ContainsByPredicate([ExistingElementKeyNode](const FRigElementResolveResult& Match)
 			{
 				return ExistingElementKeyNode->GetKey() == Match.GetKey();
-			}) ? Visible : FadedOut);
+			}) ? ESchematicGraphVisibility::Visible : ESchematicGraphVisibility::FadedOut);
 		}
 	};
 }
@@ -794,7 +879,7 @@ void FControlRigSchematicModel::HandleSchematicEndDrag(SSchematicGraphPanel* InP
 
 	for (const TSharedPtr<FSchematicGraphNode>& Node : Nodes)
 	{
-		Node->SetVisibility(Visible);
+		Node->SetVisibility(ESchematicGraphVisibility::Visible);
 	}
 }
 
diff --git a/Engine/Plugins/Animation/ControlRig/Source/ControlRigEditor/Private/Editor/ControlRigSchematicModel.h b/Engine/Plugins/Animation/ControlRig/Source/ControlRigEditor/Private/Editor/ControlRigSchematicModel.h
--- a/Engine/Plugins/Animation/ControlRig/Source/ControlRigEditor/Private/Editor/ControlRigSchematicModel.h
+++ b/Engine/Plugins/Animation/ControlRig/Source/ControlRigEditor/Private/Editor/ControlRigSchematicModel.h
@@ -77,11 +77,14 @@ public:
 
 	virtual FVector2d GetPositionForNode(const FSchematicGraphNode* InNode) const override;
 	virtual bool GetPositionAnimationEnabledForNode(const FSchematicGraphNode* InNode) const override;
-	virtual const FSlateBrush* GetBrushForNode(const FSchematicGraphNode* InNode) const override;
-	virtual FLinearColor GetColorForNode(const FSchematicGraphNode* InNode) const override;
-	virtual const FText GetToolTipForNode(const FSchematicGraphNode* InNode) const override;
-	virtual ESchematicGraphVisibility GetVisibilityForNode(const FSchematicGraphNode* InNode) const override;
-	virtual ESchematicGraphNodePlacementConstraint GetPlacementForNode(const FSchematicGraphNode* InNode) const override;
+	virtual int32 GetNumLayersForNode(const FSchematicGraphNode* InNode) const override;
+	virtual const FSlateBrush* GetBrushForNode(const FSchematicGraphNode* InNode, int32 InLayerIndex) const override;
+	virtual FLinearColor GetColorForNode(const FSchematicGraphNode* InNode, int32 InLayerIndex) const override;
+	virtual FText GetToolTipForNode(const FSchematicGraphNode* InNode) const override;
+	virtual ESchematicGraphVisibility::Type GetVisibilityForNode(const FSchematicGraphNode* InNode) const override;
+	virtual ESchematicGraphPlacementConstraint::Type GetPlacementForNode(const FSchematicGraphNode* InNode) const override;
+	virtual const FSlateBrush* GetBrushForLink(const FSchematicGraphLink* InLink) const override;
+	virtual FLinearColor GetColorForLink(const FSchematicGraphLink* InLink) const override;
 
 	virtual bool GetForwardedNodeForDrag(FGuid& InOutGuid) const override;
 
diff --git a/Engine/Source/Editor/AnimationEditorWidgets/Private/SchematicGraphPanel/SSchematicGraphPanel.cpp b/Engine/Source/Editor/AnimationEditorWidgets/Private/SchematicGraphPanel/SSchematicGraphPanel.cpp
--- a/Engine/Source/Editor/AnimationEditorWidgets/Private/SchematicGraphPanel/SSchematicGraphPanel.cpp
+++ b/Engine/Source/Editor/AnimationEditorWidgets/Private/SchematicGraphPanel/SSchematicGraphPanel.cpp
@@ -5,528 +5,11 @@
 #include "SchematicGraphPanel/SSchematicGraphPanel.h"
 #include "Framework/Application/SlateApplication.h"
 #include "DragAndDrop/AssetDragDropOp.h"
+#include "Engine/Font.h"
+#include "Engine/Engine.h"
 
 #define LOCTEXT_NAMESPACE "SSchematicGraphPanel"
 
-const FName& FSchematicGraphNode::GetType() const
-{
-	return FSchematicGraphNode::Type;
-}
-
-bool FSchematicGraphNode::IsA(const FName& InType) const
-{
-	return Type == InType;
-}
-
-FString FSchematicGraphNode::GetDragDropDecoratorLabel() const
-{
-	return GetGuid().ToString();
-}
-
- const FName& FSchematicGraphLink::GetType() const
- {
-	return FSchematicGraphLink::Type;
- }
-
- bool FSchematicGraphLink::IsA(const FName& InType) const
- {
-	return Type == InType;
- }
-
- void FSchematicGraphModel::Reset()
-{
-	Nodes.Reset();
-	Links.Reset();
-
-	if (OnGraphResetDelegate.IsBound())
-	{
-		OnGraphReset().Broadcast();
-	}
-}
-
-bool FSchematicGraphModel::RemoveNode(const FGuid& InNodeGuid)
-{
-	if(const FSchematicGraphNode* Node = FindNode(InNodeGuid))
-	{
-		if(TTuple<TArray<FGuid>,TArray<FGuid>>* ExistingTuple = NodeGuidToLinkGuids.Find(InNodeGuid))
-		{
-			TTuple<TArray<FGuid>,TArray<FGuid>> LinksToNode;
-			Swap(LinksToNode, *ExistingTuple);
-			NodeGuidToLinkGuids.Remove(InNodeGuid);
-			for(const FGuid& LinkGuid : LinksToNode.Get<0>())
-			{
-				(void)RemoveLink(LinkGuid);
-			}
-			for(const FGuid& LinkGuid : LinksToNode.Get<1>())
-			{
-				(void)RemoveLink(LinkGuid);
-			}
-		}
-		
-		if (OnNodeRemovedDelegate.IsBound())
-		{
-			OnNodeRemovedDelegate.Broadcast(Node);
-		}
-		const FGuid Guid = Node->GetGuid();
-		NodeByGuid.Remove(Node->GetGuid());
-		Nodes.RemoveAll([Guid](const TSharedPtr<FSchematicGraphNode>& ExistingNode) -> bool
-		{
-			return ExistingNode->GetGuid() == Guid;
-		});
-		return true;
-	}
-	return false;
-}
-
-FVector2d FSchematicGraphModel::GetPositionForNode(const FGuid& InNodeGuid) const
-{
-	if(const FSchematicGraphNode* Node = FindNode(InNodeGuid))
-	{
-		return GetPositionForNode(Node);
-	}
-	return FVector2d::ZeroVector;
-}
-
-FVector2d FSchematicGraphModel::GetPositionForNode(const FSchematicGraphNode* InNode) const
-{
-	check(InNode);
-	return InNode->GetPosition();
-}
-
-FVector2d FSchematicGraphModel::GetPositionOffsetForNode(const FGuid& InNodeGuid) const
-{
-	if(const FSchematicGraphNode* Node = FindNode(InNodeGuid))
-	{
-		return GetPositionOffsetForNode(Node);
-	}
-	return FVector2d::ZeroVector;
-}
-
-FVector2d FSchematicGraphModel::GetPositionOffsetForNode(const FSchematicGraphNode* InNode) const
-{
-	check(InNode);
-	return InNode->GetPositionOffset();
-}
-
-bool FSchematicGraphModel::GetPositionAnimationEnabledForNode(const FGuid& InNodeGuid) const
-{
-	if(const FSchematicGraphNode* Node = FindNode(InNodeGuid))
-	{
-		return GetPositionAnimationEnabledForNode(Node);
-	}
-	return false;
-}
-
-bool FSchematicGraphModel::GetPositionAnimationEnabledForNode(const FSchematicGraphNode* InNode) const
-{
-	return true;
-}
-
-FVector2d FSchematicGraphModel::GetSizeForNode(const FGuid& InNodeGuid) const
-{
-	if(const FSchematicGraphNode* Node = FindNode(InNodeGuid))
-	{
-		return GetSizeForNode(Node);
-	}
-	return SSchematicGraphNode::DefaultNodeSize;
-}
-
-FVector2d FSchematicGraphModel::GetSizeForNode(const FSchematicGraphNode* InNode) const
-{
-	return SSchematicGraphNode::DefaultNodeSize;
-}
-
-float FSchematicGraphModel::GetScaleForNode(const FGuid& InNodeGuid, bool bIncludeScaleOffset) const
-{
-	if(const FSchematicGraphNode* Node = FindNode(InNodeGuid))
-	{
-		return GetScaleForNode(Node, bIncludeScaleOffset);
-	}
-	return 1.f;
-}
-
-float FSchematicGraphModel::GetScaleForNode(const FSchematicGraphNode* InNode, bool bIncludeScaleOffset) const
-{
-	if(bIncludeScaleOffset)
-	{
-		return GetScaleOffsetForNode(InNode);
-	}
-	return 1.f;
-}
-
-float FSchematicGraphModel::GetScaleOffsetForNode(const FGuid& InNodeGuid) const
-{
-	if(const FSchematicGraphNode* Node = FindNode(InNodeGuid))
-	{
-		return GetScaleOffsetForNode(Node);
-	}
-	return 1.f;
-}
-
-float FSchematicGraphModel::GetScaleOffsetForNode(const FSchematicGraphNode* InNode) const
-{
-	check(InNode);
-	return InNode->GetScaleOffset();
-}
-
- float FSchematicGraphModel::GetMinimumLinkDistanceForNode(const FGuid& InNodeGuid) const
- {
-	if(const FSchematicGraphNode* Node = FindNode(InNodeGuid))
-	{
-		return GetMinimumLinkDistanceForNode(Node);
-	}
-	return 0;
- }
-
- float FSchematicGraphModel::GetMinimumLinkDistanceForNode(const FSchematicGraphNode* InNode) const
- {
-	check(InNode);
-
-	// return the radius of the node - 1 pixel
-	const FVector2d Size = GetSizeForNode(InNode);
-	return Size.GetMax() * 0.5f - 1.f;
- }
-
- bool FSchematicGraphModel::IsAutoScaleEnabledForNode(const FGuid& InNodeGuid) const
-{
-	if(const FSchematicGraphNode* Node = FindNode(InNodeGuid))
-	{
-		return IsAutoScaleEnabledForNode(Node);
-	}
-	return false;
-}
-
-bool FSchematicGraphModel::IsAutoScaleEnabledForNode(const FSchematicGraphNode* InNode) const
-{
-	check(InNode);
-	return InNode->IsAutoScaleEnabled();
-}
-
-FLinearColor FSchematicGraphModel::GetColorForNode(const FGuid& InNodeGuid) const
-{
-	if(const FSchematicGraphNode* Node = FindNode(InNodeGuid))
-	{
-		const FLinearColor Color = GetColorForNode(Node);
-		if(GetVisibilityForNode(Node) == ESchematicGraphVisibility::FadedOut)
-		{
-			return Color * 0.5f;
-		}
-		return Color;
-	}
-	return FLinearColor::White;
-}
-
-FLinearColor FSchematicGraphModel::GetColorForNode(const FSchematicGraphNode* InNode) const
-{
-	check(InNode);
-	return InNode->GetColor();
-}
-
-const FSlateBrush* FSchematicGraphModel::GetBrushForNode(const FGuid& InNodeGuid) const
-{
-	if(const FSchematicGraphNode* Node = FindNode(InNodeGuid))
-	{
-		if(const FSlateBrush* Brush = GetBrushForNode(Node))
-		{
-			return Brush;
-		}
-	}
-
-	static const FSlateBrush* DefaultBrush = SSchematicGraphNode::FArguments()._Brush.Get();
-	return DefaultBrush;
-}
-
-const FSlateBrush* FSchematicGraphModel::GetBrushForNode(const FSchematicGraphNode* InNode) const
-{
-	check(InNode);
-	return InNode->GetBrush();
-}
-
-const FText FSchematicGraphModel::GetToolTipForNode(const FGuid& InNodeGuid) const
-{
-	if(const FSchematicGraphNode* Node = FindNode(InNodeGuid))
-	{
-		return GetToolTipForNode(Node);
-	}
-	return FText();
-}
-
-const FText FSchematicGraphModel::GetToolTipForNode(const FSchematicGraphNode* InNode) const
-{
-	check(InNode);
-	return InNode->GetToolTip();
-}
-
-ESchematicGraphNodePlacementConstraint FSchematicGraphModel::GetPlacementForNode(const FGuid& InNodeGuid) const
-{
-	if(const FSchematicGraphNode* Node = FindNode(InNodeGuid))
-	{
-		return GetPlacementForNode(Node);
-	}
-	return ESchematicGraphNodePlacementConstraint::Free;
-}
-
-ESchematicGraphNodePlacementConstraint FSchematicGraphModel::GetPlacementForNode(const FSchematicGraphNode* InNode) const
-{
-	check(InNode);
-	return InNode->GetPlacement();
-}
-
-ESchematicGraphVisibility FSchematicGraphModel::GetVisibilityForNode(const FGuid& InNodeGuid) const
-{
-	if(const FSchematicGraphNode* Node = FindNode(InNodeGuid))
-	{
-		return GetVisibilityForNode(Node);
-	}
-	return ESchematicGraphVisibility::Visible;
-}
-
-ESchematicGraphVisibility FSchematicGraphModel::GetVisibilityForNode(const FSchematicGraphNode* InNode) const
-{
-	check(InNode);
-	return InNode->GetVisibility();
-}
-
-bool FSchematicGraphModel::IsDragSupportedForNode(const FGuid& InNodeGuid) const
-{
-	if(const FSchematicGraphNode* Node = FindNode(InNodeGuid))
-	{
-		return IsDragSupportedForNode(Node);
-	}
-	return false;
-}
-
-bool FSchematicGraphModel::IsDragSupportedForNode(const FSchematicGraphNode* InNode) const
-{
-	check(InNode);
-	return InNode->IsDragSupported();
-}
-
-bool FSchematicGraphModel::IsLinkedTo(const FGuid& InSourceNodeGuid, const FGuid& InTargetNodeGuid) const
-{
-	return FindLink<>(InSourceNodeGuid, InTargetNodeGuid) != nullptr;
-}
-
-TArray<const FSchematicGraphLink*> FSchematicGraphModel::FindLinksOnNode(const FGuid& InNodeGuid) const
-{
-	TArray<const FSchematicGraphLink*> Result;
-	for(const TSharedPtr<FSchematicGraphLink>& Link : Links)
-	{
-		if(Link->GetSourceNodeGuid() == InNodeGuid ||
-			Link->GetTargetNodeGuid() == InNodeGuid)
-		{
-			Result.Add(Link.Get());
-		}
-	}
-	return Result;
-}
-
-TArray<const FSchematicGraphLink*> FSchematicGraphModel::FindLinksOnSource(const FGuid& InSourceNodeGuid) const
-{
-	TArray<const FSchematicGraphLink*> Result;
-	for(const TSharedPtr<FSchematicGraphLink>& Link : Links)
-	{
-		if(Link->GetSourceNodeGuid() == InSourceNodeGuid)
-		{
-			Result.Add(Link.Get());
-		}
-	}
-	return Result;
-}
-
-TArray<const FSchematicGraphLink*> FSchematicGraphModel::FindLinksOnTarget(const FGuid& InTargetNodeGuid) const
-{
-	TArray<const FSchematicGraphLink*> Result;
-	for(const TSharedPtr<FSchematicGraphLink>& Link : Links)
-	{
-		if(Link->GetTargetNodeGuid() == InTargetNodeGuid)
-		{
-			Result.Add(Link.Get());
-		}
-	}
-	return Result;
-}
-
-bool FSchematicGraphModel::RemoveLink(const FGuid& InLinkGuid)
-{
-	if(const FSchematicGraphLink* Link = FindLink(InLinkGuid))
-	{
-		if (OnLinkRemovedDelegate.IsBound())
-		{
-			OnLinkRemovedDelegate.Broadcast(Link);
-		}
-		const FGuid Guid = Link->GetGuid();
-		const uint32 LinkHash = Link->GetLinkHash();
-		const FGuid SourceNodeGuid = Link->GetSourceNodeGuid();
-		const FGuid TargetNodeGuid = Link->GetTargetNodeGuid();
-		LinkByGuid.Remove(Link->GetGuid());
-		LinkByHash.Remove(LinkHash);
-		if(TTuple<TArray<FGuid>,TArray<FGuid>>* Tuple = NodeGuidToLinkGuids.Find(SourceNodeGuid))
-		{
-			Tuple->Get<0>().Remove(Guid);
-			if(Tuple->Get<0>().IsEmpty() && Tuple->Get<1>().IsEmpty())
-			{
-				NodeGuidToLinkGuids.Remove(SourceNodeGuid);
-			}
-		}
-		if(TTuple<TArray<FGuid>,TArray<FGuid>>* Tuple = NodeGuidToLinkGuids.Find(TargetNodeGuid))
-		{
-			Tuple->Get<1>().Remove(Guid);
-			if(Tuple->Get<0>().IsEmpty() && Tuple->Get<1>().IsEmpty())
-			{
-				NodeGuidToLinkGuids.Remove(TargetNodeGuid);
-			}
-		}
-
-		Links.RemoveAll([Guid](const TSharedPtr<FSchematicGraphLink>& ExistingLink) -> bool
-		{
-			return ExistingLink->GetGuid() == Guid;
-		});
-		return true;
-	}
-	return false;
-}
-
- float FSchematicGraphModel::GetMinimumForLink(const FGuid& InLinkGuid) const
- {
-	if(const FSchematicGraphLink* Link = FindLink(InLinkGuid))
-	{
-		return GetMinimumForLink(Link);
-	}
-	return 0.f;
- }
-
- float FSchematicGraphModel::GetMinimumForLink(const FSchematicGraphLink* InLink) const
- {
-	check(InLink);
-	return InLink->GetMinimum();
- }
-
- float FSchematicGraphModel::GetMaximumForLink(const FGuid& InLinkGuid) const
- {
-	if(const FSchematicGraphLink* Link = FindLink(InLinkGuid))
-	{
-		return GetMaximumForLink(Link);
-	}
-	return 1.f;
- }
-
- float FSchematicGraphModel::GetMaximumForLink(const FSchematicGraphLink* InLink) const
- {
-	check(InLink);
-	return InLink->GetMaximum();
- }
-
-FVector2d FSchematicGraphModel::GetSourceNodeOffsetForLink(const FGuid& InLinkGuid) const
- {
-	if(const FSchematicGraphLink* Link = FindLink(InLinkGuid))
-	{
-		return GetSourceNodeOffsetForLink(Link);
-	}
-	return FVector2d::ZeroVector;
- }
-
- FVector2d FSchematicGraphModel::GetSourceNodeOffsetForLink(const FSchematicGraphLink* InLink) const
- {
-	check(InLink);
-	return InLink->GetSourceNodeOffset();
- }
-
- FVector2d FSchematicGraphModel::GetTargetNodeOffsetForLink(const FGuid& InLinkGuid) const
- {
-	if(const FSchematicGraphLink* Link = FindLink(InLinkGuid))
-	{
-		return GetTargetNodeOffsetForLink(Link);
-	}
-	return FVector2d::ZeroVector;
- }
-
- FVector2d FSchematicGraphModel::GetTargetNodeOffsetForLink(const FSchematicGraphLink* InLink) const
- {
-	check(InLink);
-	return InLink->GetTargetNodeOffset();
- }
-
- FLinearColor FSchematicGraphModel::GetColorForLink(const FGuid& InLinkGuid) const
-{
-	if(const FSchematicGraphLink* Link = FindLink(InLinkGuid))
-	{
-		const FLinearColor Color = GetColorForLink(Link);
-		if(GetVisibilityForLink(Link) == ESchematicGraphVisibility::FadedOut)
-		{
-			return Color * 0.5f;
-		}
-		return Color;
-	}
-	return FLinearColor::White;
-}
-
-FLinearColor FSchematicGraphModel::GetColorForLink(const FSchematicGraphLink* InLink) const
-{
-	check(InLink);
-	return InLink->GetColor();
-}
-
- float FSchematicGraphModel::GetThicknessForLink(const FGuid& InLinkGuid) const
- {
-	if(const FSchematicGraphLink* Link = FindLink(InLinkGuid))
-	{
-		return GetThicknessForLink(Link);
-	}
-	return 1.f;
- }
-
- float FSchematicGraphModel::GetThicknessForLink(const FSchematicGraphLink* InLink) const
- {
-	check(InLink);
-	return InLink->GetThickness();
- }
-
- const FSlateBrush* FSchematicGraphModel::GetBrushForLink(const FGuid& InLinkGuid) const
-{
-	if(const FSchematicGraphLink* Link = FindLink(InLinkGuid))
-	{
-		return GetBrushForLink(Link);
-	}
-	return nullptr;
-}
-
-const FSlateBrush* FSchematicGraphModel::GetBrushForLink(const FSchematicGraphLink* InLink) const
-{
-	check(InLink);
-	return InLink->GetBrush();
-}
-
-const FText FSchematicGraphModel::GetToolTipForLink(const FGuid& InLinkGuid) const
-{
-	if(const FSchematicGraphLink* Link = FindLink(InLinkGuid))
-	{
-		return GetToolTipForLink(Link);
-	}
-	return FText();
-}
-
-const FText FSchematicGraphModel::GetToolTipForLink(const FSchematicGraphLink* InLink) const
-{
-	check(InLink);
-	return InLink->GetToolTip();
-}
-
-ESchematicGraphVisibility FSchematicGraphModel::GetVisibilityForLink(const FGuid& InLinkGuid) const
-{
-	if(const FSchematicGraphLink* Link = FindLink(InLinkGuid))
-	{
-		return GetVisibilityForLink(Link);
-	}
-	return ESchematicGraphVisibility::Visible;
-}
-
-ESchematicGraphVisibility FSchematicGraphModel::GetVisibilityForLink(const FSchematicGraphLink* InLink) const
-{
-	check(InLink);
-	return InLink->GetVisibility();
-}
-
 TSharedRef<FSchematicGraphNodeDragDropOp> FSchematicGraphNodeDragDropOp::New(TArray<SSchematicGraphNode*> InSchematicGraphNodes, const TArray<FGuid>& InElements, FSchematicGraphNodeDragDropOp::FOnEndDrag InOnEndDragDelegate)
 {
 	TSharedRef<FSchematicGraphNodeDragDropOp> Operation = MakeShared<FSchematicGraphNodeDragDropOp>();
@@ -575,12 +58,16 @@ FString FSchematicGraphNodeDragDropOp::GetJoinedDecoratorLabels() const
 SSchematicGraphNode::FArguments::FArguments()
 : _NodeData(nullptr)
 , _EnableAutoScale(false)
-, _Brush(nullptr)
+, _BrushGetter(nullptr)
 {
 	static const FSchematicGraphNode EmptyNodeData;
 	_NodeData = &EmptyNodeData;
-	static const FSlateBrush* WhiteTexture = FAppStyle::GetBrush("WhiteTexture");
-	_Brush = WhiteTexture;
+
+	_BrushGetter = [](const FGuid&,int32) -> const FSlateBrush*
+	{
+		static const FSlateBrush* WhiteTexture = FAppStyle::GetBrush("WhiteTexture");
+		return WhiteTexture;
+	};
 }
 
 void SSchematicGraphNode::Construct(const FArguments& InArgs)
@@ -598,11 +85,8 @@ void SSchematicGraphNode::Construct(const FArguments& InArgs)
 	Size = InArgs._Size;
 	Scale = InArgs._Scale;
 	EnableAutoScale = InArgs._EnableAutoScale;
-	Color = InArgs._Color;
-	if(InArgs._Brush.IsSet() || InArgs._Brush.IsBound())
-	{
-		Brush = InArgs._Brush;
-	}
+	LayerColors = InArgs._LayerColors;
+	BrushGetter = InArgs._BrushGetter;
 	OriginalSize = Size->Get();
 
 	if(InArgs._ToolTipText.IsBound() || InArgs._ToolTipText.IsSet())
@@ -621,19 +105,108 @@ FVector2D SSchematicGraphNode::ComputeDesiredSize(float LayoutScaleMultiplier) c
 int32 SSchematicGraphNode::OnPaint(const FPaintArgs& Args, const FGeometry& AllottedGeometry, const FSlateRect& MyCullingRect, FSlateWindowElementList& OutDrawElements, int32 LayerId, const FWidgetStyle& InWidgetStyle, bool bParentEnabled) const
 {
 	int32 NewLayerId = SNodePanel::SNode::OnPaint(Args, AllottedGeometry, MyCullingRect, OutDrawElements, LayerId, InWidgetStyle, bParentEnabled);
-	NewLayerId++;
 
 	const FVector2d CurSize = Size->Get() * Scale->Get();
 	const FVector2d SizeOffset = (CurSize-OriginalSize)*-0.5;
 
-	FSlateDrawElement::MakeBox(
+	if(BrushGetter)
+	{
+		for(int32 LayerIndex = 0; LayerIndex < LayerColors.Num(); LayerIndex++)
+		{
+			NewLayerId++;
+
+			FSlateDrawElement::MakeBox(
 				OutDrawElements,
 				NewLayerId,
 				AllottedGeometry.ToPaintGeometry(CurSize, FSlateLayoutTransform(SizeOffset)),
-				Brush.Get(),
+				BrushGetter(GetGuid(), LayerIndex),
 				ESlateDrawEffect::None,
-				Color.IsValid() ? Color->Get() : FLinearColor::White
+				LayerColors[LayerIndex].IsValid() ? LayerColors[LayerIndex]->Get() : FLinearColor::White
+			);
+		}
+	}
+
+	if(NodeData && SchematicGraphPanel && SchematicGraphPanel->GraphData)
+	{
+		const float NodeRadius = FMath::Min(CurSize.X, CurSize.Y) * 0.5;
+		
+		const TArray<TSharedPtr<FSchematicGraphTag>>& Tags = NodeData->GetTags();
+		for(const TSharedPtr<FSchematicGraphTag>& TagPtr : Tags)
+		{
+			const FSchematicGraphTag* CurrentTag = TagPtr.Get();
+			const ESchematicGraphVisibility::Type TagVisibility = SchematicGraphPanel->GraphData->GetVisibilityForTag(CurrentTag);
+			if(TagVisibility == ESchematicGraphVisibility::Hidden)
+			{
+				continue;
+			}
+
+			FVector2d TagSize = FVector2d::One() * NodeRadius * 1.0f;
+			FVector2d LabelSize = FVector2d::ZeroVector;
+			const FText& Label = SchematicGraphPanel->GraphData->GetLabelForTag(CurrentTag);
+			const FString LabelString = Label.ToString();
+
+			static const UFont* Font = GEngine->GetSmallFont();
+			static const FSlateFontInfo SmallFontStyle = FAppStyle::Get().GetFontStyle("SmallFont");
+
+			if(!LabelString.IsEmpty())
+			{
+				int32 Height = 0, Width = 0;
+				Font->GetStringHeightAndWidth(LabelString, Height, Width);
+				LabelSize.X = Width;
+				LabelSize.Y = Height;
+				TagSize.X = FMath::Max(TagSize.X, LabelSize.X + 4);
+				TagSize.Y = FMath::Max(TagSize.Y, LabelSize.Y + 4);
+			}
+
+			const FVector2d TagCenter = SizeOffset + CurSize * 0.5 + FVector2d(0, NodeRadius).GetRotated(CurrentTag->GetPlacementAngle());
+			const FVector2d TagOffset = TagCenter - TagSize * 0.5;
+
+			if(const FSlateBrush* BackgroundBrush = CurrentTag->GetBackgroundBrush())
+			{
+				NewLayerId++;
+
+				FSlateDrawElement::MakeBox(
+					OutDrawElements,
+					NewLayerId,
+					AllottedGeometry.ToPaintGeometry(TagSize, FSlateLayoutTransform(TagOffset)),
+					BackgroundBrush,
+					ESlateDrawEffect::None,
+					SchematicGraphPanel->GraphData->GetBackgroundColorForTag(CurrentTag)
 				);
+			}
+
+			if(const FSlateBrush* ForegroundBrush = SchematicGraphPanel->GraphData->GetForegroundBrushForTag(CurrentTag))
+			{
+				NewLayerId++;
+
+				FSlateDrawElement::MakeBox(
+					OutDrawElements,
+					NewLayerId,
+					AllottedGeometry.ToPaintGeometry(TagSize, FSlateLayoutTransform(TagOffset)),
+					ForegroundBrush,
+					ESlateDrawEffect::None,
+					SchematicGraphPanel->GraphData->GetForegroundColorForTag(CurrentTag)
+				);
+			}
+
+			if(!LabelSize.IsNearlyZero())
+			{
+				NewLayerId++;
+				
+				const FVector2d LabelOffset = TagCenter - LabelSize * 0.5 + FVector2d(1, 2);
+
+				FSlateDrawElement::MakeText(
+					OutDrawElements,
+					NewLayerId,
+					AllottedGeometry.ToPaintGeometry(LabelSize, FSlateLayoutTransform(LabelOffset)),
+					LabelString,
+					SmallFontStyle,
+					ESlateDrawEffect::None,
+					SchematicGraphPanel->GraphData->GetLabelColorForTag(CurrentTag)
+				);
+			}
+		}
+	}
 	
 	return NewLayerId;
 }
@@ -740,9 +313,9 @@ FReply SSchematicGraphNode::OnDragDetected(const FGeometry& MyGeometry, const FP
 	TArray<FGuid> DraggedElements = {Guid};
 	if (MouseEvent.IsMouseButtonDown(EKeys::LeftMouseButton) && DraggedElements.Num() > 0)
 	{
-		if(SchematicGraphPanel)
+		if(SchematicGraphPanel && SchematicGraphPanel->GraphData)
 		{
-			if(SchematicGraphPanel->IsDragSupportedForNode(GetGuid()))
+			if(SchematicGraphPanel->GraphData->IsDragSupportedForNode(GetGuid()))
 			{
 				bIsBeingDragged = true;
 
@@ -767,9 +340,9 @@ EVisibility SSchematicGraphNode::GetNodeVisibility() const
 		return EVisibility::HitTestInvisible;
 	}
 
-	if(SchematicGraphPanel)
+	if(SchematicGraphPanel && SchematicGraphPanel->GraphData)
 	{
-		ESchematicGraphVisibility Vis = SchematicGraphPanel->GetVisibilityForNode(GetGuid());
+		ESchematicGraphVisibility::Type Vis = SchematicGraphPanel->GraphData->GetVisibilityForNode(GetGuid());
 		if(Vis == ESchematicGraphVisibility::Hidden)
 		{
 			return EVisibility::Hidden;
@@ -891,24 +464,40 @@ void SSchematicGraphPanel::RebuildPanel()
 
 void SSchematicGraphPanel::AddNode(const FSchematicGraphNode* InNodeToAdd)
 {
+	if(GraphData == nullptr)
+	{
+		return;
+	}
+
 	static const TEasingAttributeInterpolator<FVector2d>::FSettings Vector2DInterpolationSettings(EEasingInterpolatorType::CubicEaseOut, 0.2f);
 	static const TEasingAttributeInterpolator<float>::FSettings FloatInterpolationSettings(EEasingInterpolatorType::CubicEaseOut, 0.2f);
 	static const TEasingAttributeInterpolator<FLinearColor>::FSettings ColorInterpolationSettings(EEasingInterpolatorType::CubicEaseOut, 0.2f);
 
 	const FGuid Guid = InNodeToAdd->GetGuid();
 	const auto Position = FVector2dAttribute::CreateWithGetter(Vector2DInterpolationSettings, FVector2dAttribute::FGetter::CreateSP(this, &SSchematicGraphPanel::GetPositionForNode, Guid));
-	const auto Size = FVector2dAttribute::CreateWithGetter(Vector2DInterpolationSettings, FVector2dAttribute::FGetter::CreateSP(this, &SSchematicGraphPanel::GetSizeForNode, Guid));
+	const auto Size = FVector2dAttribute::CreateWithGetter(Vector2DInterpolationSettings, FVector2dAttribute::FGetter::CreateRaw(GraphData, &FSchematicGraphModel::GetSizeForNode, InNodeToAdd));
 	const auto Scale = FFloatAttribute::CreateWithGetter(FloatInterpolationSettings, FFloatAttribute::FGetter::CreateSP(this, &SSchematicGraphPanel::GetScaleForNode, Guid, true), 0.f);
-	const auto Color = FLinearColorAttribute::CreateWithGetter(ColorInterpolationSettings, FLinearColorAttribute::FGetter::CreateSP(this, &SSchematicGraphPanel::GetColorForNode, Guid));
+
+	const int32 NumLayers = GraphData ? GraphData->GetNumLayersForNode(InNodeToAdd) : InNodeToAdd->GetNumLayers();
+	TArray<TSharedPtr<FLinearColorAttribute>> Colors;
+	for(int32 LayerIndex = 0; LayerIndex < NumLayers; LayerIndex++)
+	{
+		Colors.Add(FLinearColorAttribute::CreateWithGetter(ColorInterpolationSettings, FLinearColorAttribute::FGetter::CreateRaw(GraphData, &FSchematicGraphModel::GetColorForNode, InNodeToAdd, LayerIndex)));
+	}
+
+	const TFunction<const FSlateBrush*(const FGuid&, int32)> BrushGetter = [this](const FGuid& InGuid, int32 InLayerIndex) -> const FSlateBrush*
+	{
+		return GraphData->GetBrushForNode(InGuid, InLayerIndex);
+	};
 
 	const TSharedRef<SSchematicGraphNode> NewNode = SNew(SSchematicGraphNode)
 														.Position(Position)
 														.Size(Size)
 														.Scale(Scale)
-														.Color(Color)
+														.LayerColors(Colors)
 														.EnableAutoScale(this, &SSchematicGraphPanel::IsAutoScaleEnabledForNode, InNodeToAdd->GetGuid())
-														.Brush(this, &SSchematicGraphPanel::GetBrushForNode, InNodeToAdd->GetGuid())
-														.ToolTipText(this, &SSchematicGraphPanel::GetToolTipForNode, InNodeToAdd->GetGuid())
+														.BrushGetter(BrushGetter)
+														.ToolTipText_Raw(GraphData, &FSchematicGraphModel::GetToolTipForNode, InNodeToAdd)
 														.OnClicked_Raw(this, &SSchematicGraphPanel::OnNodeClicked)
 														.OnBeginDrag_Raw(this, &SSchematicGraphPanel::OnBeginDragEvent)
 														.OnEndDrag_Raw(this, &SSchematicGraphPanel::OnEndDragEvent)
@@ -955,15 +544,20 @@ const SSchematicGraphNode* SSchematicGraphPanel::FindNode(const FGuid& InGuid) c
 
  void SSchematicGraphPanel::AddLink(const FSchematicGraphLink* InLinkToAdd)
  {
+	if(GraphData == nullptr)
+	{
+		return;
+	}
+	
 	static const TEasingAttributeInterpolator<float>::FSettings FloatInterpolationSettings(EEasingInterpolatorType::CubicEaseOut, 0.1f);
 	static const TEasingAttributeInterpolator<float>::FSettings SlowFloatInterpolationSettings(EEasingInterpolatorType::CubicEaseOut, 0.2f);
 	static const TEasingAttributeInterpolator<FLinearColor>::FSettings ColorInterpolationSettings(EEasingInterpolatorType::CubicEaseOut, 0.2f);
 
 	const FGuid Guid = InLinkToAdd->GetGuid();
-	const auto Minimum = FFloatAttribute::CreateWithGetter(SlowFloatInterpolationSettings, FFloatAttribute::FGetter::CreateSP(this, &SSchematicGraphPanel::GetMinimumForLink, Guid), 0.5f);
-	const auto Maximum = FFloatAttribute::CreateWithGetter(SlowFloatInterpolationSettings, FFloatAttribute::FGetter::CreateSP(this, &SSchematicGraphPanel::GetMaximumForLink, Guid), 0.5f);
-	const auto Color = FLinearColorAttribute::CreateWithGetter(ColorInterpolationSettings, FLinearColorAttribute::FGetter::CreateSP(this, &SSchematicGraphPanel::GetColorForLink, Guid));
-	const auto Thickness = FFloatAttribute::CreateWithGetter(FloatInterpolationSettings, FFloatAttribute::FGetter::CreateSP(this, &SSchematicGraphPanel::GetThicknessForLink, Guid), 0);
+	const auto Minimum = FFloatAttribute::CreateWithGetter(SlowFloatInterpolationSettings, FFloatAttribute::FGetter::CreateRaw(GraphData, &FSchematicGraphModel::GetMinimumForLink, InLinkToAdd), 0.5f);
+	const auto Maximum = FFloatAttribute::CreateWithGetter(SlowFloatInterpolationSettings, FFloatAttribute::FGetter::CreateRaw(GraphData, &FSchematicGraphModel::GetMaximumForLink, InLinkToAdd), 0.5f);
+	const auto Color = FLinearColorAttribute::CreateWithGetter(ColorInterpolationSettings, FLinearColorAttribute::FGetter::CreateRaw(GraphData, &FSchematicGraphModel::GetColorForLink, InLinkToAdd));
+	const auto Thickness = FFloatAttribute::CreateWithGetter(FloatInterpolationSettings, FFloatAttribute::FGetter::CreateRaw(GraphData, &FSchematicGraphModel::GetThicknessForLink, InLinkToAdd), 0);
 
 	FSchematicLinkWidgetInfo Info;
 	Info.Minimum = Minimum;
@@ -1059,7 +653,7 @@ int32 SSchematicGraphPanel::OnPaint(const FPaintArgs& Args, const FGeometry& All
 	// draw all of the links
 	for(const TPair<FGuid, TSharedPtr<FSchematicLinkWidgetInfo>>& Pair : LinkByGuid)
 	{
-		if(GetVisibilityForLink(Pair.Key) == ESchematicGraphVisibility::Hidden)
+		if(GraphData->GetVisibilityForLink(Pair.Key) == ESchematicGraphVisibility::Hidden)
 		{
 			continue;
 		}
@@ -1091,7 +685,7 @@ int32 SSchematicGraphPanel::OnPaint(const FPaintArgs& Args, const FGeometry& All
 
 		const FLinearColor Color = Pair.Value->Color->Get();
 		const float Thickness = Pair.Value->Thickness->Get();
-		const FSlateBrush* Brush = GetBrushForLink(Pair.Key);
+		const FSlateBrush* Brush = GraphData->GetBrushForLink(Pair.Key);
 		const FVector2d& SourcePosition = NodeCenterByGuid.FindChecked(SourceNode->GetGuid()) + GraphData->GetSourceNodeOffsetForLink(Link);
 		const FVector2d& TargetPosition = NodeCenterByGuid.FindChecked(TargetNode->GetGuid()) + GraphData->GetTargetNodeOffsetForLink(Link);
 
@@ -1122,17 +716,41 @@ int32 SSchematicGraphPanel::OnPaint(const FPaintArgs& Args, const FGeometry& All
 			FMath::Lerp<FVector2d>(MinimumPosition, MaximumPosition, FMath::Clamp(Maximum, 0, 1))
 		};
 
-		// todo: draw the lines using a brush
-		FSlateDrawElement::MakeLines(
-			OutDrawElements,
-			LinkLayerId,
-			AllottedGeometry.ToPaintGeometry(),
-			LinePoints,
-			ESlateDrawEffect::None,
-			Color,
-			true,
-			Thickness
-		);
+		if(Brush == nullptr)
+		{
+			FSlateDrawElement::MakeLines(
+				OutDrawElements,
+				LinkLayerId,
+				AllottedGeometry.ToPaintGeometry(),
+				LinePoints,
+				ESlateDrawEffect::None,
+				Color,
+				true,
+				Thickness
+			);
+		}
+		else
+		{
+			const FVector2d Center = (LinePoints[0] + LinePoints[1]) * 0.5f;
+			const float Distance = (LinePoints[0] - LinePoints[1]).Size();
+			static constexpr float DefaultDistance = 128.f;
+			const float AdjustedThickness = Thickness * FMath::Min(1, Distance / DefaultDistance);
+			const FVector2d LineSize = {AdjustedThickness, Distance };
+			const FVector2d SizeOffset = LineSize * 0.5f;
+			const float Angle = -FMath::Atan2(-Diff.X, -Diff.Y);
+
+			FSlateDrawElement::MakeRotatedBox(
+				OutDrawElements,
+				LinkLayerId,
+				AllottedGeometry.ToPaintGeometry(LineSize, FSlateLayoutTransform(Center - SizeOffset)),
+				Brush,
+				ESlateDrawEffect::None,
+				Angle,
+				SizeOffset, // rotation point
+				FSlateDrawElement::ERotationSpace::RelativeToElement,
+				Color
+				);
+		}
 	}
 
 	// Because we paint multiple children, we must track the maximum layer id that they produced in case one of our parents
@@ -1218,26 +836,26 @@ void SSchematicGraphPanel::Tick(float DeltaTime)
 		// collect the nodes displayed in the corners
 		switch(GraphData->GetPlacementForNode(Node->GetNodeData()))
 		{
-			case ESchematicGraphNodePlacementConstraint::Free:
+			case ESchematicGraphPlacementConstraint::Free:
 			{
 				break;
 			}
-			case ESchematicGraphNodePlacementConstraint::TopLeft:
+			case ESchematicGraphPlacementConstraint::TopLeft:
 			{
 				NodesTopLeft.Add(Node->GetGuid());
 				break;
 			}
-			case ESchematicGraphNodePlacementConstraint::TopRight:
+			case ESchematicGraphPlacementConstraint::TopRight:
 			{
 				NodesTopRight.Add(Node->GetGuid());
 				break;
 			}
-			case ESchematicGraphNodePlacementConstraint::BottomLeft:
+			case ESchematicGraphPlacementConstraint::BottomLeft:
 			{
 				NodesBottomLeft.Add(Node->GetGuid());
 				break;
 			}
-			case ESchematicGraphNodePlacementConstraint::BottomRight:
+			case ESchematicGraphPlacementConstraint::BottomRight:
 			{
 				NodesBottomRight.Add(Node->GetGuid());
 				break;
@@ -1327,12 +945,12 @@ FVector2d SSchematicGraphPanel::GetPositionForNode(FGuid InNodeGuid) const
 
 			switch(GraphData->GetPlacementForNode(Node))
 			{
-				case ESchematicGraphNodePlacementConstraint::Free:
+				case ESchematicGraphPlacementConstraint::Free:
 				{
 					Position += GraphData->GetPositionForNode(Node);
 					break;
 				}
-				case ESchematicGraphNodePlacementConstraint::TopLeft:
+				case ESchematicGraphPlacementConstraint::TopLeft:
 				{
 					const int32 NodeIndexInCorner = NodesTopLeft.IndexOfByKey(InNodeGuid);
 					if(NodeIndexInCorner != INDEX_NONE)
@@ -1341,7 +959,7 @@ FVector2d SSchematicGraphPanel::GetPositionForNode(FGuid InNodeGuid) const
 					}
 					break;
 				}
-				case ESchematicGraphNodePlacementConstraint::TopRight:
+				case ESchematicGraphPlacementConstraint::TopRight:
 				{
 					const int32 NodeIndexInCorner = NodesTopRight.IndexOfByKey(InNodeGuid);
 					if(NodeIndexInCorner != INDEX_NONE)
@@ -1350,7 +968,7 @@ FVector2d SSchematicGraphPanel::GetPositionForNode(FGuid InNodeGuid) const
 					}
 					break;
 				}
-				case ESchematicGraphNodePlacementConstraint::BottomLeft:
+				case ESchematicGraphPlacementConstraint::BottomLeft:
 				{
 					const int32 NodeIndexInCorner = NodesBottomLeft.IndexOfByKey(InNodeGuid);
 					if(NodeIndexInCorner != INDEX_NONE)
@@ -1359,7 +977,7 @@ FVector2d SSchematicGraphPanel::GetPositionForNode(FGuid InNodeGuid) const
 					}
 					break;
 				}
-				case ESchematicGraphNodePlacementConstraint::BottomRight:
+				case ESchematicGraphPlacementConstraint::BottomRight:
 				{
 					const int32 NodeIndexInCorner = NodesBottomRight.IndexOfByKey(InNodeGuid);
 					if(NodeIndexInCorner != INDEX_NONE)
@@ -1375,15 +993,6 @@ FVector2d SSchematicGraphPanel::GetPositionForNode(FGuid InNodeGuid) const
 	return FVector2d::ZeroVector;
 }
 
-FVector2d SSchematicGraphPanel::GetSizeForNode(FGuid InNodeGuid) const
-{
-	if(GraphData)
-	{
-		return GraphData->GetSizeForNode(InNodeGuid);
-	}
-	return SSchematicGraphNode::DefaultNodeSize;
-}
-
 float SSchematicGraphPanel::GetScaleForNode(FGuid InNodeGuid, bool bIncludeScaleOffset) const
 {
 	// check if the node may be auto scaled
@@ -1431,118 +1040,13 @@ bool SSchematicGraphPanel::IsAutoScaleEnabledForNode(FGuid InNodeGuid) const
 	return 0.f;
  }
 
- FLinearColor SSchematicGraphPanel::GetColorForNode(FGuid InNodeGuid) const
-{
-	if(GraphData)
-	{
-		return GraphData->GetColorForNode(InNodeGuid);
-	}
-	return FLinearColor::White;
-}
-
-const FSlateBrush* SSchematicGraphPanel::GetBrushForNode(FGuid InNodeGuid) const
-{
-	if(GraphData)
-	{
-		return GraphData->GetBrushForNode(InNodeGuid);
-	}
-	static const FSlateBrush* DefaultBrush = SSchematicGraphNode::FArguments()._Brush.Get();
-	return DefaultBrush;
-}
-
-FText SSchematicGraphPanel::GetToolTipForNode(FGuid InNodeGuid) const
-{
-	if(GraphData)
-	{
-		return GraphData->GetToolTipForNode(InNodeGuid);
-	}
-	return FText();
-}
-
-ESchematicGraphVisibility SSchematicGraphPanel::GetVisibilityForNode(FGuid InNodeGuid) const
-{
-	if(GraphData)
-	{
-		return GraphData->GetVisibilityForNode(InNodeGuid);
-	}
-	return ESchematicGraphVisibility::Visible;
-}
-
-bool SSchematicGraphPanel::IsDragSupportedForNode(FGuid InNodeGuid) const
-{
-	if(GraphData)
-	{
-		return GraphData->IsDragSupportedForNode(InNodeGuid);
-	}
-	return false;
-}
-
- float SSchematicGraphPanel::GetMinimumForLink(FGuid InLinkGuid) const
- {
-	if(GraphData)
-	{
-		return GraphData->GetMinimumForLink(InLinkGuid);
-	}
-	return 1.f;
- }
-
- float SSchematicGraphPanel::GetMaximumForLink(FGuid InLinkGuid) const
- {
-	if(GraphData)
-	{
-		return GraphData->GetMaximumForLink(InLinkGuid);
-	}
-	return 1.f;
- }
-
- FLinearColor SSchematicGraphPanel::GetColorForLink(FGuid InLinkGuid) const
-{
-	if(GraphData)
-	{
-		return GraphData->GetColorForLink(InLinkGuid);
-	}
-	return FLinearColor::White;
-}
-
- float SSchematicGraphPanel::GetThicknessForLink(FGuid InLinkGuid) const
- {
-	if(GraphData)
-	{
-		return GraphData->GetThicknessForLink(InLinkGuid);
-	}
-	return 1.f;
- }
-
- const FSlateBrush* SSchematicGraphPanel::GetBrushForLink(FGuid InLinkGuid) const
-{
-	if(GraphData)
-	{
-		return GraphData->GetBrushForLink(InLinkGuid);
-	}
-	static const FSlateBrush* DefaultBrush = nullptr;
-	return DefaultBrush;
-}
-
-FText SSchematicGraphPanel::GetToolTipForLink(FGuid InLinkGuid) const
-{
-	if(GraphData)
-	{
-		return GraphData->GetToolTipForLink(InLinkGuid);
-	}
-	return FText();
-}
-
-ESchematicGraphVisibility SSchematicGraphPanel::GetVisibilityForLink(FGuid InLinkGuid) const
+void SSchematicGraphPanel::UpdateAutoScalingForNodes()
 {
-	if(GraphData)
+	if(GraphData == nullptr)
 	{
-		return GraphData->GetVisibilityForLink(InLinkGuid);
+		return;
 	}
-	return ESchematicGraphVisibility::Visible;
-}
-
-void SSchematicGraphPanel::UpdateAutoScalingForNodes()
-{
+	
 	TArray<bool> NodeIsAutoScaling;
 	TArray<FVector2d> NodePositions;
 	TArray<double> NodeRadiuses;
@@ -1558,7 +1062,7 @@ void SSchematicGraphPanel::UpdateAutoScalingForNodes()
 		if(NodeIsAutoScaling.Last())
 		{
 			NodePositions.Add(GetPositionForNode(Node->GetGuid()));
-			const FVector2d NodeSize = GetSizeForNode(Node->GetGuid());
+			const FVector2d NodeSize = GraphData->GetSizeForNode(Node->GetGuid());
 			// note: this is not necessarily the best way to determine the radius of a node
 			NodeRadiuses.Add(FMath::Min(NodeSize.X, NodeSize.Y) * 0.5);
 		}
diff --git a/Engine/Source/Editor/AnimationEditorWidgets/Private/SchematicGraphPanel/SchematicGraphModel.cpp b/Engine/Source/Editor/AnimationEditorWidgets/Private/SchematicGraphPanel/SchematicGraphModel.cpp
new file mode 100644
--- /dev/null
+++ b/Engine/Source/Editor/AnimationEditorWidgets/Private/SchematicGraphPanel/SchematicGraphModel.cpp
@@ -0,0 +1,731 @@
+// Copyright Epic Games, Inc. All Rights Reserved.
+
+#if WITH_EDITOR
+
+#include "SchematicGraphPanel/SchematicGraphModel.h"
+#include "SchematicGraphPanel/SSchematicGraphPanel.h"
+
+#define LOCTEXT_NAMESPACE "SchematicGraphModel"
+
+void FSchematicGraphModel::Reset()
+{
+	Nodes.Reset();
+	Links.Reset();
+
+	if (OnGraphResetDelegate.IsBound())
+	{
+		OnGraphReset().Broadcast();
+	}
+}
+
+bool FSchematicGraphModel::RemoveNode(const FGuid& InNodeGuid)
+{
+	if(const FSchematicGraphNode* Node = FindNode(InNodeGuid))
+	{
+		if(TTuple<TArray<FGuid>,TArray<FGuid>>* ExistingTuple = NodeGuidToLinkGuids.Find(InNodeGuid))
+		{
+			TTuple<TArray<FGuid>,TArray<FGuid>> LinksToNode;
+			Swap(LinksToNode, *ExistingTuple);
+			NodeGuidToLinkGuids.Remove(InNodeGuid);
+			for(const FGuid& LinkGuid : LinksToNode.Get<0>())
+			{
+				(void)RemoveLink(LinkGuid);
+			}
+			for(const FGuid& LinkGuid : LinksToNode.Get<1>())
+			{
+				(void)RemoveLink(LinkGuid);
+			}
+		}
+		
+		if (OnNodeRemovedDelegate.IsBound())
+		{
+			OnNodeRemovedDelegate.Broadcast(Node);
+		}
+		const FGuid NodeGuid = Node->GetGuid();
+		NodeByGuid.Remove(Node->GetGuid());
+		Nodes.RemoveAll([NodeGuid](const TSharedPtr<FSchematicGraphNode>& ExistingNode) -> bool
+		{
+			return ExistingNode->GetGuid() == NodeGuid;
+		});
+		return true;
+	}
+	return false;
+}
+
+bool FSchematicGraphModel::SetParentNode(const FGuid& InChildNodeGuid, const FGuid& InParentNodeGuid)
+{
+	if(const FSchematicGraphNode* Node = FindNode(InChildNodeGuid))
+	{
+		return SetParentNode(Node, FindNode(InParentNodeGuid));
+	}
+	return false;
+}
+
+bool FSchematicGraphModel::SetParentNode(const FSchematicGraphNode* InChildNode, const FSchematicGraphNode* InParentNode)
+{
+	check(InChildNode);
+
+	if(InParentNode)
+	{
+		if(InParentNode == InChildNode->GetParentNode())
+		{
+			return false;
+		}
+		
+		(void)RemoveFromParentNode(InChildNode);
+		const_cast<FSchematicGraphNode*>(InParentNode)->ChildNodeGuids.AddUnique(InChildNode->GetGuid());
+		const_cast<FSchematicGraphNode*>(InChildNode)->ParentNodeGuid = InParentNode->GetGuid();
+		return true;
+	}
+	
+	if(const FSchematicGraphNode* CurrentParentNode = InChildNode->GetParentNode())
+	{
+		(void)const_cast<FSchematicGraphNode*>(CurrentParentNode)->ChildNodeGuids.Remove(InChildNode->GetGuid());
+		const_cast<FSchematicGraphNode*>(InChildNode)->ParentNodeGuid = FGuid();
+		return true;
+	}
+	
+	return false;
+}
+
+bool FSchematicGraphModel::RemoveFromParentNode(const FGuid& InChildNodeGuid)
+{
+	if(const FSchematicGraphNode* Node = FindNode(InChildNodeGuid))
+	{
+		return RemoveFromParentNode(Node);
+	}
+	return false;
+}
+
+bool FSchematicGraphModel::RemoveFromParentNode(const FSchematicGraphNode* InChildNode)
+{
+	check(InChildNode);
+	return SetParentNode(InChildNode, nullptr);
+}
+
+FVector2d FSchematicGraphModel::GetPositionForNode(const FGuid& InNodeGuid) const
+{
+	if(const FSchematicGraphNode* Node = FindNode(InNodeGuid))
+	{
+		return GetPositionForNode(Node);
+	}
+	return FVector2d::ZeroVector;
+}
+
+FVector2d FSchematicGraphModel::GetPositionForNode(const FSchematicGraphNode* InNode) const
+{
+	check(InNode);
+	return InNode->GetPosition();
+}
+
+FVector2d FSchematicGraphModel::GetPositionOffsetForNode(const FGuid& InNodeGuid) const
+{
+	if(const FSchematicGraphNode* Node = FindNode(InNodeGuid))
+	{
+		return GetPositionOffsetForNode(Node);
+	}
+	return FVector2d::ZeroVector;
+}
+
+FVector2d FSchematicGraphModel::GetPositionOffsetForNode(const FSchematicGraphNode* InNode) const
+{
+	check(InNode);
+	return InNode->GetPositionOffset();
+}
+
+bool FSchematicGraphModel::GetPositionAnimationEnabledForNode(const FGuid& InNodeGuid) const
+{
+	if(const FSchematicGraphNode* Node = FindNode(InNodeGuid))
+	{
+		return GetPositionAnimationEnabledForNode(Node);
+	}
+	return false;
+}
+
+bool FSchematicGraphModel::GetPositionAnimationEnabledForNode(const FSchematicGraphNode* InNode) const
+{
+	return true;
+}
+
+FVector2d FSchematicGraphModel::GetSizeForNode(const FGuid& InNodeGuid) const
+{
+	if(const FSchematicGraphNode* Node = FindNode(InNodeGuid))
+	{
+		return GetSizeForNode(Node);
+	}
+	return SSchematicGraphNode::DefaultNodeSize;
+}
+
+FVector2d FSchematicGraphModel::GetSizeForNode(const FSchematicGraphNode* InNode) const
+{
+	return SSchematicGraphNode::DefaultNodeSize;
+}
+
+float FSchematicGraphModel::GetScaleForNode(const FGuid& InNodeGuid, bool bIncludeScaleOffset) const
+{
+	if(const FSchematicGraphNode* Node = FindNode(InNodeGuid))
+	{
+		return GetScaleForNode(Node, bIncludeScaleOffset);
+	}
+	return 1.f;
+}
+
+float FSchematicGraphModel::GetScaleForNode(const FSchematicGraphNode* InNode, bool bIncludeScaleOffset) const
+{
+	if(bIncludeScaleOffset)
+	{
+		return GetScaleOffsetForNode(InNode);
+	}
+	return 1.f;
+}
+
+float FSchematicGraphModel::GetScaleOffsetForNode(const FGuid& InNodeGuid) const
+{
+	if(const FSchematicGraphNode* Node = FindNode(InNodeGuid))
+	{
+		return GetScaleOffsetForNode(Node);
+	}
+	return 1.f;
+}
+
+float FSchematicGraphModel::GetScaleOffsetForNode(const FSchematicGraphNode* InNode) const
+{
+	check(InNode);
+	return InNode->GetScaleOffset();
+}
+
+ float FSchematicGraphModel::GetMinimumLinkDistanceForNode(const FGuid& InNodeGuid) const
+ {
+	if(const FSchematicGraphNode* Node = FindNode(InNodeGuid))
+	{
+		return GetMinimumLinkDistanceForNode(Node);
+	}
+	return 0;
+ }
+
+ float FSchematicGraphModel::GetMinimumLinkDistanceForNode(const FSchematicGraphNode* InNode) const
+ {
+	check(InNode);
+
+	// return the radius of the node - 1 pixel
+	const FVector2d Size = GetSizeForNode(InNode);
+	return Size.GetMax() * 0.5f - 1.f;
+ }
+
+ bool FSchematicGraphModel::IsAutoScaleEnabledForNode(const FGuid& InNodeGuid) const
+{
+	if(const FSchematicGraphNode* Node = FindNode(InNodeGuid))
+	{
+		return IsAutoScaleEnabledForNode(Node);
+	}
+	return false;
+}
+
+bool FSchematicGraphModel::IsAutoScaleEnabledForNode(const FSchematicGraphNode* InNode) const
+{
+	check(InNode);
+	return InNode->IsAutoScaleEnabled();
+}
+
+ int32 FSchematicGraphModel::GetNumLayersForNode(const FGuid& InNodeGuid) const
+ {
+	if(const FSchematicGraphNode* Node = FindNode(InNodeGuid))
+	{
+		return GetNumLayersForNode(Node);
+	}
+	return 0;
+ }
+
+  int32 FSchematicGraphModel::GetNumLayersForNode(const FSchematicGraphNode* InNode) const
+  {
+	check(InNode);
+	return InNode->GetNumLayers();
+  }
+
+  FLinearColor FSchematicGraphModel::GetColorForNode(const FGuid& InNodeGuid, int32 InLayerIndex) const
+{
+	if(const FSchematicGraphNode* Node = FindNode(InNodeGuid))
+	{
+		const FLinearColor Color = GetColorForNode(Node, InLayerIndex);
+		if(GetVisibilityForNode(Node) == ESchematicGraphVisibility::FadedOut)
+		{
+			return Color * 0.5f;
+		}
+		return Color;
+	}
+	return FLinearColor::White;
+}
+
+FLinearColor FSchematicGraphModel::GetColorForNode(const FSchematicGraphNode* InNode, int32 InLayerIndex) const
+{
+	check(InNode);
+	return InNode->GetColor(InLayerIndex);
+}
+
+const FSlateBrush* FSchematicGraphModel::GetBrushForNode(const FGuid& InNodeGuid, int32 InLayerIndex) const
+{
+	if(const FSchematicGraphNode* Node = FindNode(InNodeGuid))
+	{
+		if(const FSlateBrush* Brush = GetBrushForNode(Node, InLayerIndex))
+		{
+			return Brush;
+		}
+	}
+
+	static const FSlateBrush* DefaultBrush = SSchematicGraphNode::FArguments()._BrushGetter(FGuid(), INDEX_NONE);
+	return DefaultBrush;
+}
+
+const FSlateBrush* FSchematicGraphModel::GetBrushForNode(const FSchematicGraphNode* InNode, int32 InLayerIndex) const
+{
+	check(InNode);
+	return InNode->GetBrush(InLayerIndex);
+}
+
+FText FSchematicGraphModel::GetToolTipForNode(const FGuid& InNodeGuid) const
+{
+	if(const FSchematicGraphNode* Node = FindNode(InNodeGuid))
+	{
+		return GetToolTipForNode(Node);
+	}
+	return FText();
+}
+
+FText FSchematicGraphModel::GetToolTipForNode(const FSchematicGraphNode* InNode) const
+{
+	check(InNode);
+	return InNode->GetToolTip();
+}
+
+ESchematicGraphPlacementConstraint::Type FSchematicGraphModel::GetPlacementForNode(const FGuid& InNodeGuid) const
+{
+	if(const FSchematicGraphNode* Node = FindNode(InNodeGuid))
+	{
+		return GetPlacementForNode(Node);
+	}
+	return ESchematicGraphPlacementConstraint::Free;
+}
+
+ESchematicGraphPlacementConstraint::Type FSchematicGraphModel::GetPlacementForNode(const FSchematicGraphNode* InNode) const
+{
+	check(InNode);
+	return InNode->GetPlacement();
+}
+
+ESchematicGraphVisibility::Type FSchematicGraphModel::GetVisibilityForNode(const FGuid& InNodeGuid) const
+{
+	if(const FSchematicGraphNode* Node = FindNode(InNodeGuid))
+	{
+		return GetVisibilityForNode(Node);
+	}
+	return ESchematicGraphVisibility::Visible;
+}
+
+ESchematicGraphVisibility::Type FSchematicGraphModel::GetVisibilityForNode(const FSchematicGraphNode* InNode) const
+{
+	check(InNode);
+	return InNode->GetVisibility();
+}
+
+ESchematicGraphVisibility::Type FSchematicGraphModel::GetVisibilityForChildNodes(const FGuid& InNodeGuid) const
+{
+	if(const FSchematicGraphNode* Node = FindNode(InNodeGuid))
+	{
+		return GetVisibilityForChildNodes(Node);
+	}
+	return ESchematicGraphVisibility::Visible;
+}
+
+ESchematicGraphVisibility::Type FSchematicGraphModel::GetVisibilityForChildNodes(const FSchematicGraphNode* InNode) const
+{
+	check(InNode);
+	return InNode->GetVisibilityForChildNodes();
+}
+
+bool FSchematicGraphModel::IsDragSupportedForNode(const FGuid& InNodeGuid) const
+{
+	if(const FSchematicGraphNode* Node = FindNode(InNodeGuid))
+	{
+		return IsDragSupportedForNode(Node);
+	}
+	return false;
+}
+
+bool FSchematicGraphModel::IsDragSupportedForNode(const FSchematicGraphNode* InNode) const
+{
+	check(InNode);
+	return InNode->IsDragSupported();
+}
+
+FLinearColor FSchematicGraphModel::GetBackgroundColorForTag(const FGuid& InNodeGuid, const FGuid& InTagGuid) const
+{
+	if(const FSchematicGraphNode* Node = FindNode(InNodeGuid))
+	{
+		if(const FSchematicGraphTag* Tag = Node->FindTag(InTagGuid))
+		{
+			return GetBackgroundColorForTag(Tag);
+		}
+	}
+	return FLinearColor::White;
+}
+
+FLinearColor FSchematicGraphModel::GetBackgroundColorForTag(const FSchematicGraphTag* InTag) const
+{
+	check(InTag);
+	return InTag->GetBackgroundColor();
+}
+
+FLinearColor FSchematicGraphModel::GetForegroundColorForTag(const FGuid& InNodeGuid, const FGuid& InTagGuid) const
+{
+	if(const FSchematicGraphNode* Node = FindNode(InNodeGuid))
+	{
+		if(const FSchematicGraphTag* Tag = Node->FindTag(InTagGuid))
+		{
+			return GetForegroundColorForTag(Tag);
+		}
+	}
+	return FLinearColor::White;
+}
+
+FLinearColor FSchematicGraphModel::GetForegroundColorForTag(const FSchematicGraphTag* InTag) const
+{
+	check(InTag);
+	return InTag->GetForegroundColor();
+}
+
+FLinearColor FSchematicGraphModel::GetLabelColorForTag(const FGuid& InNodeGuid, const FGuid& InTagGuid) const
+{
+	if(const FSchematicGraphNode* Node = FindNode(InNodeGuid))
+	{
+		if(const FSchematicGraphTag* Tag = Node->FindTag(InTagGuid))
+		{
+			return GetLabelColorForTag(Tag);
+		}
+	}
+	return FLinearColor::White;
+}
+
+FLinearColor FSchematicGraphModel::GetLabelColorForTag(const FSchematicGraphTag* InTag) const
+{
+	check(InTag);
+	return InTag->GetLabelColor();
+}
+
+const FSlateBrush* FSchematicGraphModel::GetBackgroundBrushForTag(const FGuid& InNodeGuid, const FGuid& InTagGuid) const
+{
+	if(const FSchematicGraphNode* Node = FindNode(InNodeGuid))
+	{
+		if(const FSchematicGraphTag* Tag = Node->FindTag(InTagGuid))
+		{
+			return GetBackgroundBrushForTag(Tag);
+		}
+	}
+	return nullptr;
+}
+
+const FSlateBrush* FSchematicGraphModel::GetBackgroundBrushForTag(const FSchematicGraphTag* InTag) const
+{
+	check(InTag);
+	return InTag->GetBackgroundBrush();
+}
+
+const FSlateBrush* FSchematicGraphModel::GetForegroundBrushForTag(const FGuid& InNodeGuid, const FGuid& InTagGuid) const
+{
+	if(const FSchematicGraphNode* Node = FindNode(InNodeGuid))
+	{
+		if(const FSchematicGraphTag* Tag = Node->FindTag(InTagGuid))
+		{
+			return GetForegroundBrushForTag(Tag);
+		}
+	}
+	return nullptr;
+}
+
+const FSlateBrush* FSchematicGraphModel::GetForegroundBrushForTag(const FSchematicGraphTag* InTag) const
+{
+	check(InTag);
+	return InTag->GetForegroundBrush();
+}
+
+const FText FSchematicGraphModel::GetLabelForTag(const FGuid& InNodeGuid, const FGuid& InTagGuid) const
+{
+	if(const FSchematicGraphNode* Node = FindNode(InNodeGuid))
+	{
+		if(const FSchematicGraphTag* Tag = Node->FindTag(InTagGuid))
+		{
+			return GetLabelForTag(Tag);
+		}
+	}
+	return FText();
+}
+
+const FText FSchematicGraphModel::GetLabelForTag(const FSchematicGraphTag* InTag) const
+{
+	check(InTag);
+	return InTag->GetLabel();
+}
+
+const FText FSchematicGraphModel::GetToolTipForTag(const FGuid& InNodeGuid, const FGuid& InTagGuid) const
+{
+	if(const FSchematicGraphNode* Node = FindNode(InNodeGuid))
+	{
+		if(const FSchematicGraphTag* Tag = Node->FindTag(InTagGuid))
+		{
+			return GetToolTipForTag(Tag);
+		}
+	}
+	return FText();
+}
+
+const FText FSchematicGraphModel::GetToolTipForTag(const FSchematicGraphTag* InTag) const
+{
+	check(InTag);
+	return InTag->GetToolTip();
+}
+
+ESchematicGraphVisibility::Type FSchematicGraphModel::GetVisibilityForTag(const FGuid& InNodeGuid, const FGuid& InTagGuid) const
+{
+	if(const FSchematicGraphNode* Node = FindNode(InNodeGuid))
+	{
+		if(const FSchematicGraphTag* Tag = Node->FindTag(InTagGuid))
+		{
+			return GetVisibilityForTag(Tag);
+		}
+	}
+	return ESchematicGraphVisibility::Visible;
+}
+
+ESchematicGraphVisibility::Type FSchematicGraphModel::GetVisibilityForTag(const FSchematicGraphTag* InTag) const
+{
+	check(InTag);
+	return InTag->GetVisibility();
+}
+
+bool FSchematicGraphModel::IsLinkedTo(const FGuid& InSourceNodeGuid, const FGuid& InTargetNodeGuid) const
+{
+	return FindLink<>(InSourceNodeGuid, InTargetNodeGuid) != nullptr;
+}
+
+TArray<const FSchematicGraphLink*> FSchematicGraphModel::FindLinksOnNode(const FGuid& InNodeGuid) const
+{
+	TArray<const FSchematicGraphLink*> Result;
+	for(const TSharedPtr<FSchematicGraphLink>& Link : Links)
+	{
+		if(Link->GetSourceNodeGuid() == InNodeGuid ||
+			Link->GetTargetNodeGuid() == InNodeGuid)
+		{
+			Result.Add(Link.Get());
+		}
+	}
+	return Result;
+}
+
+TArray<const FSchematicGraphLink*> FSchematicGraphModel::FindLinksOnSource(const FGuid& InSourceNodeGuid) const
+{
+	TArray<const FSchematicGraphLink*> Result;
+	for(const TSharedPtr<FSchematicGraphLink>& Link : Links)
+	{
+		if(Link->GetSourceNodeGuid() == InSourceNodeGuid)
+		{
+			Result.Add(Link.Get());
+		}
+	}
+	return Result;
+}
+
+TArray<const FSchematicGraphLink*> FSchematicGraphModel::FindLinksOnTarget(const FGuid& InTargetNodeGuid) const
+{
+	TArray<const FSchematicGraphLink*> Result;
+	for(const TSharedPtr<FSchematicGraphLink>& Link : Links)
+	{
+		if(Link->GetTargetNodeGuid() == InTargetNodeGuid)
+		{
+			Result.Add(Link.Get());
+		}
+	}
+	return Result;
+}
+
+bool FSchematicGraphModel::RemoveLink(const FGuid& InLinkGuid)
+{
+	if(const FSchematicGraphLink* Link = FindLink(InLinkGuid))
+	{
+		if (OnLinkRemovedDelegate.IsBound())
+		{
+			OnLinkRemovedDelegate.Broadcast(Link);
+		}
+		const FGuid LinkGuid = Link->GetGuid();
+		const uint32 LinkHash = Link->GetLinkHash();
+		const FGuid SourceNodeGuid = Link->GetSourceNodeGuid();
+		const FGuid TargetNodeGuid = Link->GetTargetNodeGuid();
+		LinkByGuid.Remove(Link->GetGuid());
+		LinkByHash.Remove(LinkHash);
+		if(TTuple<TArray<FGuid>,TArray<FGuid>>* Tuple = NodeGuidToLinkGuids.Find(SourceNodeGuid))
+		{
+			Tuple->Get<0>().Remove(LinkGuid);
+			if(Tuple->Get<0>().IsEmpty() && Tuple->Get<1>().IsEmpty())
+			{
+				NodeGuidToLinkGuids.Remove(SourceNodeGuid);
+			}
+		}
+		if(TTuple<TArray<FGuid>,TArray<FGuid>>* Tuple = NodeGuidToLinkGuids.Find(TargetNodeGuid))
+		{
+			Tuple->Get<1>().Remove(LinkGuid);
+			if(Tuple->Get<0>().IsEmpty() && Tuple->Get<1>().IsEmpty())
+			{
+				NodeGuidToLinkGuids.Remove(TargetNodeGuid);
+			}
+		}
+
+		Links.RemoveAll([LinkGuid](const TSharedPtr<FSchematicGraphLink>& ExistingLink) -> bool
+		{
+			return ExistingLink->GetGuid() == LinkGuid;
+		});
+		return true;
+	}
+	return false;
+}
+
+ float FSchematicGraphModel::GetMinimumForLink(const FGuid& InLinkGuid) const
+ {
+	if(const FSchematicGraphLink* Link = FindLink(InLinkGuid))
+	{
+		return GetMinimumForLink(Link);
+	}
+	return 0.f;
+ }
+
+ float FSchematicGraphModel::GetMinimumForLink(const FSchematicGraphLink* InLink) const
+ {
+	check(InLink);
+	return InLink->GetMinimum();
+ }
+
+ float FSchematicGraphModel::GetMaximumForLink(const FGuid& InLinkGuid) const
+ {
+	if(const FSchematicGraphLink* Link = FindLink(InLinkGuid))
+	{
+		return GetMaximumForLink(Link);
+	}
+	return 1.f;
+ }
+
+ float FSchematicGraphModel::GetMaximumForLink(const FSchematicGraphLink* InLink) const
+ {
+	check(InLink);
+	return InLink->GetMaximum();
+ }
+
+FVector2d FSchematicGraphModel::GetSourceNodeOffsetForLink(const FGuid& InLinkGuid) const
+ {
+	if(const FSchematicGraphLink* Link = FindLink(InLinkGuid))
+	{
+		return GetSourceNodeOffsetForLink(Link);
+	}
+	return FVector2d::ZeroVector;
+ }
+
+ FVector2d FSchematicGraphModel::GetSourceNodeOffsetForLink(const FSchematicGraphLink* InLink) const
+ {
+	check(InLink);
+	return InLink->GetSourceNodeOffset();
+ }
+
+ FVector2d FSchematicGraphModel::GetTargetNodeOffsetForLink(const FGuid& InLinkGuid) const
+ {
+	if(const FSchematicGraphLink* Link = FindLink(InLinkGuid))
+	{
+		return GetTargetNodeOffsetForLink(Link);
+	}
+	return FVector2d::ZeroVector;
+ }
+
+ FVector2d FSchematicGraphModel::GetTargetNodeOffsetForLink(const FSchematicGraphLink* InLink) const
+ {
+	check(InLink);
+	return InLink->GetTargetNodeOffset();
+ }
+
+ FLinearColor FSchematicGraphModel::GetColorForLink(const FGuid& InLinkGuid) const
+{
+	if(const FSchematicGraphLink* Link = FindLink(InLinkGuid))
+	{
+		const FLinearColor Color = GetColorForLink(Link);
+		if(GetVisibilityForLink(Link) == ESchematicGraphVisibility::FadedOut)
+		{
+			return Color * 0.5f;
+		}
+		return Color;
+	}
+	return FLinearColor::White;
+}
+
+FLinearColor FSchematicGraphModel::GetColorForLink(const FSchematicGraphLink* InLink) const
+{
+	check(InLink);
+	return InLink->GetColor();
+}
+
+ float FSchematicGraphModel::GetThicknessForLink(const FGuid& InLinkGuid) const
+ {
+	if(const FSchematicGraphLink* Link = FindLink(InLinkGuid))
+	{
+		return GetThicknessForLink(Link);
+	}
+	return 1.f;
+ }
+
+ float FSchematicGraphModel::GetThicknessForLink(const FSchematicGraphLink* InLink) const
+ {
+	check(InLink);
+	return InLink->GetThickness();
+ }
+
+ const FSlateBrush* FSchematicGraphModel::GetBrushForLink(const FGuid& InLinkGuid) const
+{
+	if(const FSchematicGraphLink* Link = FindLink(InLinkGuid))
+	{
+		return GetBrushForLink(Link);
+	}
+	return nullptr;
+}
+
+const FSlateBrush* FSchematicGraphModel::GetBrushForLink(const FSchematicGraphLink* InLink) const
+{
+	check(InLink);
+	return InLink->GetBrush();
+}
+
+const FText FSchematicGraphModel::GetToolTipForLink(const FGuid& InLinkGuid) const
+{
+	if(const FSchematicGraphLink* Link = FindLink(InLinkGuid))
+	{
+		return GetToolTipForLink(Link);
+	}
+	return FText();
+}
+
+const FText FSchematicGraphModel::GetToolTipForLink(const FSchematicGraphLink* InLink) const
+{
+	check(InLink);
+	return InLink->GetToolTip();
+}
+
+ESchematicGraphVisibility::Type FSchematicGraphModel::GetVisibilityForLink(const FGuid& InLinkGuid) const
+{
+	if(const FSchematicGraphLink* Link = FindLink(InLinkGuid))
+	{
+		return GetVisibilityForLink(Link);
+	}
+	return ESchematicGraphVisibility::Visible;
+}
+
+ESchematicGraphVisibility::Type FSchematicGraphModel::GetVisibilityForLink(const FSchematicGraphLink* InLink) const
+{
+	check(InLink);
+	return InLink->GetVisibility();
+}
+
+#undef LOCTEXT_NAMESPACE
+
+#endif
\ No newline at end of file
diff --git a/Engine/Source/Editor/AnimationEditorWidgets/Private/SchematicGraphPanel/SchematicGraphNode.cpp b/Engine/Source/Editor/AnimationEditorWidgets/Private/SchematicGraphPanel/SchematicGraphNode.cpp
new file mode 100644
--- /dev/null
+++ b/Engine/Source/Editor/AnimationEditorWidgets/Private/SchematicGraphPanel/SchematicGraphNode.cpp
@@ -0,0 +1,92 @@
+// Copyright Epic Games, Inc. All Rights Reserved.
+
+#if WITH_EDITOR
+
+#include "SchematicGraphPanel/SchematicGraphNode.h"
+#include "SchematicGraphPanel/SchematicGraphModel.h"
+
+#define LOCTEXT_NAMESPACE "SchematicGraphNode"
+
+const FSchematicGraphNode* FSchematicGraphNode::GetParentNode() const
+{
+	if(Model && ParentNodeGuid.IsValid())
+	{
+		return Model->FindNode(ParentNodeGuid);
+	}
+	return nullptr;
+}
+
+const FSchematicGraphNode* FSchematicGraphNode::GetChildNode(int32 InChildNodeIndex) const
+{
+	check(ChildNodeGuids.IsValidIndex(InChildNodeIndex));
+
+	const FGuid& ChildNodeGuid = ChildNodeGuids[InChildNodeIndex];
+	if(Model && ChildNodeGuid.IsValid())
+	{
+		return Model->FindNode(ChildNodeGuid);
+	}
+	return nullptr;
+}
+
+ESchematicGraphVisibility::Type FSchematicGraphNode::GetVisibility() const
+{
+	if(Visibility == ESchematicGraphVisibility::Hidden)
+	{
+		return Visibility;
+	}
+	
+	if(const FSchematicGraphNode* ParentNode = GetParentNode())
+	{
+		const ESchematicGraphVisibility::Type ParentVisibility = Model->GetVisibilityForChildNodes(ParentNode);
+		if(ParentVisibility != ESchematicGraphVisibility::Visible)
+		{
+			return ParentVisibility;
+		}
+	}
+	return Visibility;
+}
+
+FString FSchematicGraphNode::GetDragDropDecoratorLabel() const
+{
+	return GetGuid().ToString();
+}
+
+bool FSchematicGraphNode::RemoveTag(const FGuid& InTagGuid)
+{
+	if(const FSchematicGraphTag* Tag = FindTag(InTagGuid))
+	{
+		if (Model && Model->OnTagRemovedDelegate.IsBound())
+		{
+			Model->OnTagRemovedDelegate.Broadcast(this, Tag);
+		}
+		const FGuid TagGuid = Tag->GetGuid();
+		TagByGuid.Remove(Tag->GetGuid());
+		Tags.RemoveAll([TagGuid](const TSharedPtr<FSchematicGraphTag>& ExistingTag) -> bool
+		{
+			return ExistingTag->GetGuid() == TagGuid;
+		});
+		return true;
+	}
+	return false;
+}
+
+void FSchematicGraphNode::NotifyTagAdded(const TSharedPtr<FSchematicGraphTag>& Tag)
+{
+	if (Model != nullptr && Model->OnTagAddedDelegate.IsBound())
+	{
+		Model->OnTagAddedDelegate.Broadcast(this, Tag.Get());
+	}
+}
+
+ESchematicGraphVisibility::Type FSchematicGraphGroupNode::GetVisibilityForChildNodes() const
+{
+	if(!IsExpanded())
+	{
+		return ESchematicGraphVisibility::Hidden;
+	}
+	return FSchematicGraphNode::GetVisibilityForChildNodes();
+}
+
+#undef LOCTEXT_NAMESPACE
+
+#endif
\ No newline at end of file
diff --git a/Engine/Source/Editor/AnimationEditorWidgets/Private/SchematicGraphPanel/SchematicGraphTag.cpp b/Engine/Source/Editor/AnimationEditorWidgets/Private/SchematicGraphPanel/SchematicGraphTag.cpp
new file mode 100644
--- /dev/null
+++ b/Engine/Source/Editor/AnimationEditorWidgets/Private/SchematicGraphPanel/SchematicGraphTag.cpp
@@ -0,0 +1,21 @@
+// Copyright Epic Games, Inc. All Rights Reserved.
+
+#if WITH_EDITOR
+
+#include "SchematicGraphPanel/SchematicGraphTag.h"
+#include "SchematicGraphPanel/SchematicGraphNode.h"
+
+#define LOCTEXT_NAMESPACE "SchematicGraphTag"
+
+const FSchematicGraphNode* FSchematicGraphTag::GetNode() const
+{
+	if(Node.IsValid())
+	{
+		return Node.Pin().Get();
+	}
+	return nullptr;
+}
+
+#undef LOCTEXT_NAMESPACE
+
+#endif
\ No newline at end of file
diff --git a/Engine/Source/Editor/AnimationEditorWidgets/Public/SchematicGraphPanel/SSchematicGraphPanel.h b/Engine/Source/Editor/AnimationEditorWidgets/Public/SchematicGraphPanel/SSchematicGraphPanel.h
--- a/Engine/Source/Editor/AnimationEditorWidgets/Public/SchematicGraphPanel/SSchematicGraphPanel.h
+++ b/Engine/Source/Editor/AnimationEditorWidgets/Public/SchematicGraphPanel/SSchematicGraphPanel.h
@@ -3,451 +3,12 @@
 
 #if WITH_EDITOR
 
-#include "CoreMinimal.h"
+#include "SchematicGraphModel.h"
 #include "SNodePanel.h"
 #include "TickableEditorObject.h"
 #include "Animation/AnimatedAttribute.h"
 
 class SSchematicGraphPanel;
-
-enum ESchematicGraphNodePlacementConstraint
-{
-	Free,
-	TopLeft,
-	TopRight,
-	BottomLeft,
-	BottomRight
-};
-
-enum ESchematicGraphVisibility
-{
-	Visible,
-	FadedOut,
-	Hidden
-};
-
-#define SCHEMATICGRAPHELEMENT_BODY(ClassName, SuperClass, BaseClass) \
-inline static const FName& Type = TEXT(#ClassName); \
-virtual const FName& GetType() const override { return ClassName::Type; } \
-virtual bool IsA(const FName& InType) const \
-{ \
-	if(ClassName::Type == InType) \
-	{ \
-		return true; \
-	} \
-	return SuperClass::IsA(InType); \
-} \
-template<typename T> \
-friend const T* Cast(const ClassName* InNode) \
-{ \
-	return Cast<T>((const BaseClass*) InNode); \
-} \
-template<typename T> \
-friend T* Cast(ClassName* InNode) \
-{ \
-	return Cast<T>((BaseClass*) InNode); \
-} \
-template<typename T> \
-friend const T* CastChecked(const ClassName* InNode) \
-{ \
-	return CastChecked<T>((const BaseClass*) InNode); \
-} \
-template<typename T> \
-friend T* CastChecked(ClassName* InNode) \
-{ \
-	return CastChecked<T>((BaseClass*) InNode); \
-}
-
-#define SCHEMATICGRAPHNODE_BODY(ClassName, SuperClass) \
-SCHEMATICGRAPHELEMENT_BODY(ClassName, SuperClass, FSchematicGraphNode)
-
-#define SCHEMATICGRAPHLINK_BODY(ClassName, SuperClass) \
-SCHEMATICGRAPHELEMENT_BODY(ClassName, SuperClass, FSchematicGraphNode)
-
-class ANIMATIONEDITORWIDGETS_API FSchematicGraphNode
-{
-public:
-	virtual ~FSchematicGraphNode() {}
-	
-	inline static const FName& Type = TEXT("FSchematicGraphNode");
-	virtual const FName& GetType() const;
-	virtual bool IsA(const FName& InType) const;
-	template<typename NodeType>
-	bool IsA() const
-	{
-		return IsA(NodeType::Type);
-	}
-	
-	template<typename T>
-	friend const T* Cast(const FSchematicGraphNode* InNode)
-	{
-		if(InNode)
-		{
-			if(InNode->IsA<T>())
-			{
-				return static_cast<const T*>(InNode);
-			}
-		}
-		return nullptr;
-	}
-
-	template<typename T>
-	friend T* Cast(FSchematicGraphNode* InNode)
-	{
-		if(InNode)
-		{
-			if(InNode->IsA<T>())
-			{
-				return static_cast<T*>(InNode);
-			}
-		}
-		return nullptr;
-	}
-
-	template<typename T>
-	friend const T* CastChecked(const FSchematicGraphNode* InNode)
-	{
-		const T* Node = Cast<T>(InNode);
-		check(Node);
-		return Node;
-	}
-
-	template<typename T>
-	friend T* CastChecked(FSchematicGraphNode* InNode)
-	{
-		T* Node = Cast<T>(InNode);
-		check(Node);
-		return Node;
-	}
-	
-	const FGuid& GetGuid() const { return Guid; }
-	bool IsSelected() const { return bIsSelected; }
-	void SetSelected(bool bSelected = true) { bIsSelected = bSelected; }
-	virtual FVector2d GetPosition() const { return Position; }
-	virtual FVector2d GetPositionOffset() const { return PositionOffset; }
-	virtual void SetPositionOffset(const FVector2d& InPositionOffset) { PositionOffset = InPositionOffset; }
-	virtual float GetScaleOffset() const { return ScaleOffset; }
-	virtual void SetScaleOffset(float InScaleOffset) { ScaleOffset = InScaleOffset; }
-	virtual FLinearColor GetColor() const { return Color; }
-	virtual const FSlateBrush* GetBrush() const { return Brush; }
-	virtual bool IsAutoScaleEnabled() const { return false; }
-	virtual const FText& GetToolTip() const { return ToolTip; }
-	virtual ESchematicGraphNodePlacementConstraint GetPlacement() const { return Placement; }
-	virtual void SetPlacement(ESchematicGraphNodePlacementConstraint InPlacement) { Placement = InPlacement; }
-	virtual ESchematicGraphVisibility GetVisibility() const { return Visibility; }
-	virtual void SetVisibility(ESchematicGraphVisibility InVisibility) { Visibility = InVisibility; }
-	virtual bool IsDragSupported() const { return bDragSupported; }
-	virtual void SetDragSupported(bool InDragSupported) { bDragSupported = InDragSupported;}
-
-	virtual FString GetDragDropDecoratorLabel() const;
-
-protected:
-	
-	FGuid Guid = FGuid::NewGuid();
-	bool bIsSelected = false;
-	FVector2d Position = FVector2d::ZeroVector;
-	FVector2d PositionOffset = FVector2d::ZeroVector;
-	float ScaleOffset = 1.f;
-	FLinearColor Color = FLinearColor::White;
-	const FSlateBrush* Brush = nullptr;
-	FText ToolTip = FText();
-	ESchematicGraphNodePlacementConstraint Placement = ESchematicGraphNodePlacementConstraint::Free;
-	ESchematicGraphVisibility Visibility = ESchematicGraphVisibility::Visible;
-	bool bDragSupported = false;
-
-	friend class FSchematicGraphModel;
-};
-
-class ANIMATIONEDITORWIDGETS_API FSchematicGraphLink
-{
-public:
-	virtual ~FSchematicGraphLink() {}
-	
-	inline static const FName& Type = TEXT("FSchematicGraphLink");
-	virtual const FName& GetType() const;
-	virtual bool IsA(const FName& InType) const;
-	template<typename NodeType>
-	bool IsA() const
-	{
-		return IsA(NodeType::Type);
-	}
-	
-	template<typename T>
-	friend const T* Cast(const FSchematicGraphLink* InNode)
-	{
-		if(InNode)
-		{
-			if(InNode->IsA<T>())
-			{
-				return static_cast<const T*>(InNode);
-			}
-		}
-		return nullptr;
-	}
-
-	template<typename T>
-	friend T* Cast(FSchematicGraphLink* InNode)
-	{
-		if(InNode)
-		{
-			if(InNode->IsA<T>())
-			{
-				return static_cast<T*>(InNode);
-			}
-		}
-		return nullptr;
-	}
-
-	template<typename T>
-	friend const T* CastChecked(const FSchematicGraphLink* InNode)
-	{
-		const T* Node = Cast<T>(InNode);
-		check(Node);
-		return Node;
-	}
-
-	template<typename T>
-	friend T* CastChecked(FSchematicGraphLink* InNode)
-	{
-		T* Node = Cast<T>(InNode);
-		check(Node);
-		return Node;
-	}
-	
-	const FGuid& GetGuid() const { return Guid; }
-	static uint32 GetLinkHash(const FGuid& InSourceNodeGuid, const FGuid& InTargetNodeGuid)
-	{
-		return HashCombine(GetTypeHash(InSourceNodeGuid), GetTypeHash(InTargetNodeGuid));
-	}
-	uint32 GetLinkHash() const { return GetLinkHash(SourceNodeGuid, TargetNodeGuid); }
-	const FGuid& GetSourceNodeGuid() const { return SourceNodeGuid; }
-	const FGuid& GetTargetNodeGuid() const { return TargetNodeGuid; }
-	virtual float GetMinimum() const { return Minimum; }
-	virtual float GetMaximum() const { return Maximum; }
-	virtual FVector2d GetSourceNodeOffset() const { return SourceNodeOffset; }
-	virtual FVector2d GetTargetNodeOffset() const { return TargetNodeOffset; }
-	virtual FLinearColor GetColor() const { return Color; }
-	virtual float GetThickness() const { return Thickness; }
-	virtual const FSlateBrush* GetBrush() const { return Brush; }
-	virtual const FText& GetToolTip() const { return ToolTip; }
-	virtual ESchematicGraphVisibility GetVisibility() const { return Visibility; }
-	virtual void SetVisibility(ESchematicGraphVisibility InVisibility) { Visibility = InVisibility; }
-
-protected:
-	
-	FGuid Guid = FGuid::NewGuid();
-	FGuid SourceNodeGuid = FGuid();
-	FGuid TargetNodeGuid = FGuid();
-	float Minimum = 0.f;
-	float Maximum = 1.f;
-	FVector2d SourceNodeOffset = FVector2d::ZeroVector;
-	FVector2d TargetNodeOffset = FVector2d::ZeroVector;
-	FLinearColor Color = FLinearColor::White;
-	float Thickness = 1.f;
-	const FSlateBrush* Brush = nullptr;
-	FText ToolTip = FText();
-	ESchematicGraphVisibility Visibility = ESchematicGraphVisibility::Visible;
-
-	friend class FSchematicGraphModel;
-};
-
-DECLARE_EVENT_OneParam(FSchematicGraph, FOnNodeAdded, const FSchematicGraphNode*);
-DECLARE_EVENT_OneParam(FSchematicGraph, FOnNodeRemoved, const FSchematicGraphNode*);
-DECLARE_EVENT_OneParam(FSchematicGraph, FOnLinkAdded, const FSchematicGraphLink*);
-DECLARE_EVENT_OneParam(FSchematicGraph, FOnLinkRemoved, const FSchematicGraphLink*);
-DECLARE_EVENT(FSchematicGraph, FOnGraphReset);
-
-class ANIMATIONEDITORWIDGETS_API FSchematicGraphModel
-{
-public:
-
-	virtual ~FSchematicGraphModel() {}
-	
-	virtual void Reset();
-	
-	template<typename NodeType = FSchematicGraphNode>
-	NodeType* AddNode(bool bNotify = true)
-	{
-		const TSharedPtr<FSchematicGraphNode> NewNode = MakeShareable(new NodeType);
-		Nodes.Add(NewNode);
-		NodeByGuid.Add(NewNode->GetGuid(), NewNode);
-
-		if (bNotify && OnNodeAddedDelegate.IsBound())
-		{
-			OnNodeAddedDelegate.Broadcast(NewNode.Get());
-		}
-		return static_cast<NodeType*>(NewNode.Get());
-	}
-
-	template<typename NodeType = FSchematicGraphNode>
-	const NodeType* FindNode(const FGuid& InNodeGuid) const
-	{
-		if(const TSharedPtr<FSchematicGraphNode>* ExistingNode = NodeByGuid.Find(InNodeGuid))
-		{
-			return Cast<NodeType>(ExistingNode->Get());
-		};
-		return nullptr;
-	}
-
-	template<typename NodeType = FSchematicGraphNode>
-	const NodeType* FindNodeChecked(const FGuid& InNodeGuid) const
-	{
-		if(const TSharedPtr<FSchematicGraphNode>* ExistingNode = NodeByGuid.Find(InNodeGuid))
-		{
-			check(ExistingNode->Get()->IsA(NodeType::Type));
-			return static_cast<NodeType*>(ExistingNode->Get());
-		};
-		return nullptr;
-	}
-
-	virtual bool RemoveNode(const FGuid& InNodeGuid);
-
-	const TArray<TSharedPtr<FSchematicGraphNode>>& GetNodes() const { return Nodes; }
-
-	FVector2d GetPositionForNode(const FGuid& InNodeGuid) const;
-	virtual FVector2d GetPositionForNode(const FSchematicGraphNode* InNode) const;
-	FVector2d GetPositionOffsetForNode(const FGuid& InNodeGuid) const;
-	virtual FVector2d GetPositionOffsetForNode(const FSchematicGraphNode* InNode) const;
-	bool GetPositionAnimationEnabledForNode(const FGuid& InNodeGuid) const;
-	virtual bool GetPositionAnimationEnabledForNode(const FSchematicGraphNode* InNode) const;
-	FVector2d GetSizeForNode(const FGuid& InNodeGuid) const;
-	virtual FVector2d GetSizeForNode(const FSchematicGraphNode* InNode) const;
-	float GetScaleForNode(const FGuid& InNodeGuid, bool bIncludeScaleOffset) const;
-	virtual float GetScaleForNode(const FSchematicGraphNode* InNode, bool bIncludeScaleOffset) const;
-	float GetScaleOffsetForNode(const FGuid& InNodeGuid) const;
-	virtual float GetScaleOffsetForNode(const FSchematicGraphNode* InNode) const;
-	float GetMinimumLinkDistanceForNode(const FGuid& InNodeGuid) const;
-	virtual float GetMinimumLinkDistanceForNode(const FSchematicGraphNode* InNode) const;
-	bool IsAutoScaleEnabledForNode(const FGuid& InNodeGuid) const;
-	virtual bool IsAutoScaleEnabledForNode(const FSchematicGraphNode* InNode) const;
-	FLinearColor GetColorForNode(const FGuid& InNodeGuid) const;
-	virtual FLinearColor GetColorForNode(const FSchematicGraphNode* InNode) const;
-	const FSlateBrush* GetBrushForNode(const FGuid& InNodeGuid) const;
-	virtual const FSlateBrush* GetBrushForNode(const FSchematicGraphNode* InNode) const;
-	const FText GetToolTipForNode(const FGuid& InNodeGuid) const;
-	virtual const FText GetToolTipForNode(const FSchematicGraphNode* InNode) const;
-	ESchematicGraphNodePlacementConstraint GetPlacementForNode(const FGuid& InNodeGuid) const;
-	virtual ESchematicGraphNodePlacementConstraint GetPlacementForNode(const FSchematicGraphNode* InNode) const;
-	ESchematicGraphVisibility GetVisibilityForNode(const FGuid& InNodeGuid) const;
-	virtual ESchematicGraphVisibility GetVisibilityForNode(const FSchematicGraphNode* InNode) const;
-	bool IsDragSupportedForNode(const FGuid& InNodeGuid) const;
-	virtual bool IsDragSupportedForNode(const FSchematicGraphNode* InNode) const;
-
-	template<typename LinkType = FSchematicGraphLink>
-	LinkType* AddLink(const FGuid& InSourceNodeGuid, const FGuid& InTargetNodeGuid, bool bNotify = true)
-	{
-		const TSharedPtr<FSchematicGraphLink> NewLink = MakeShareable(new LinkType);
-		NewLink->SourceNodeGuid = InSourceNodeGuid;
-		NewLink->TargetNodeGuid = InTargetNodeGuid;
-		Links.Add(NewLink);
-		LinkByGuid.Add(NewLink->GetGuid(), NewLink);
-		LinkByHash.Add(NewLink->GetLinkHash(), NewLink);
-		NodeGuidToLinkGuids.FindOrAdd(InSourceNodeGuid).Get<0>().Add(NewLink->GetGuid());
-		NodeGuidToLinkGuids.FindOrAdd(InTargetNodeGuid).Get<1>().Add(NewLink->GetGuid());
-
-		if (bNotify && OnLinkAddedDelegate.IsBound())
-		{
-			OnLinkAddedDelegate.Broadcast(NewLink.Get());
-		}
-		return static_cast<LinkType*>(NewLink.Get());
-	}
-
-	template<typename LinkType = FSchematicGraphLink>
-	const LinkType* FindLink(const FGuid& InLinkGuid) const
-	{
-		if(const TSharedPtr<FSchematicGraphLink>* ExistingLink = LinkByGuid.Find(InLinkGuid))
-		{
-			return Cast<LinkType>(ExistingLink->Get());
-		};
-		return nullptr;
-	}
-
-	template<typename LinkType = FSchematicGraphLink>
-	const LinkType* FindLinkChecked(const FGuid& InLinkGuid) const
-	{
-		if(const TSharedPtr<FSchematicGraphLink>* ExistingLink = LinkByGuid.Find(InLinkGuid))
-		{
-			check(ExistingLink->Get()->IsA(LinkType::Type));
-			return static_cast<LinkType*>(ExistingLink->Get());
-		};
-		return nullptr;
-	}
-
-	template<typename LinkType = FSchematicGraphLink>
-	const LinkType* FindLink(const FGuid& InSourceNodeGuid, const FGuid& InTargetNodeGuid) const
-	{
-		const uint32 LinkHash = FSchematicGraphLink::GetLinkHash(InSourceNodeGuid, InTargetNodeGuid);
-		if(const TSharedPtr<FSchematicGraphLink>* ExistingLink = LinkByHash.Find(LinkHash))
-		{
-			return Cast<LinkType>(ExistingLink->Get());
-		};
-		return nullptr;
-	}
-
-	template<typename LinkType = FSchematicGraphLink>
-	const LinkType* FindLinkChecked(const FGuid& InSourceNodeGuid, const FGuid& InTargetNodeGuid) const
-	{
-		const uint32 LinkHash = FSchematicGraphLink::GetLinkHash(InSourceNodeGuid, InTargetNodeGuid);
-		if(const TSharedPtr<FSchematicGraphLink>* ExistingLink = LinkByHash.Find(LinkHash))
-		{
-			check(ExistingLink->Get()->IsA(LinkType::Type));
-			return Cast<LinkType>(ExistingLink->Get());
-		};
-		return nullptr;
-	}
-
-	bool IsLinkedTo(const FGuid& InSourceNodeGuid, const FGuid& InTargetNodeGuid) const;
-	TArray<const FSchematicGraphLink*> FindLinksOnNode(const FGuid& InNodeGuid) const;
-	TArray<const FSchematicGraphLink*> FindLinksOnSource(const FGuid& InSourceNodeGuid) const;
-	TArray<const FSchematicGraphLink*> FindLinksOnTarget(const FGuid& InTargetNodeGuid) const;
-
-	virtual bool RemoveLink(const FGuid& InLinkGuid);
-
-	const TArray<TSharedPtr<FSchematicGraphLink>>& GetLinks() const { return Links; }
-
-	float GetMinimumForLink(const FGuid& InLinkGuid) const;
-	virtual float GetMinimumForLink(const FSchematicGraphLink* InLink) const;
-	float GetMaximumForLink(const FGuid& InLinkGuid) const;
-	virtual float GetMaximumForLink(const FSchematicGraphLink* InLink) const;
-	FVector2d GetSourceNodeOffsetForLink(const FGuid& InLinkGuid) const;
-	FVector2d GetSourceNodeOffsetForLink(const FSchematicGraphLink* InLink) const;
-	FVector2d GetTargetNodeOffsetForLink(const FGuid& InLinkGuid) const;
-	FVector2d GetTargetNodeOffsetForLink(const FSchematicGraphLink* InLink) const;
-	FLinearColor GetColorForLink(const FGuid& InLinkGuid) const;
-	virtual FLinearColor GetColorForLink(const FSchematicGraphLink* InLink) const;
-	float GetThicknessForLink(const FGuid& InLinkGuid) const;
-	virtual float GetThicknessForLink(const FSchematicGraphLink* InLink) const;
-	const FSlateBrush* GetBrushForLink(const FGuid& InLinkGuid) const;
-	virtual const FSlateBrush* GetBrushForLink(const FSchematicGraphLink* InLink) const;
-	const FText GetToolTipForLink(const FGuid& InLinkGuid) const;
-	virtual const FText GetToolTipForLink(const FSchematicGraphLink* InLink) const;
-	ESchematicGraphVisibility GetVisibilityForLink(const FGuid& InLinkGuid) const;
-	virtual ESchematicGraphVisibility GetVisibilityForLink(const FSchematicGraphLink* InLink) const;
-
-	FOnNodeAdded& OnNodeAdded() { return OnNodeAddedDelegate; }
-	FOnNodeRemoved& OnNodeRemoved() { return OnNodeRemovedDelegate; }
-	FOnLinkAdded& OnLinkAdded() { return OnLinkAddedDelegate; }
-	FOnLinkRemoved& OnLinkRemoved() { return OnLinkRemovedDelegate; }
-	FOnGraphReset& OnGraphReset() { return OnGraphResetDelegate; }
-
-	virtual bool GetForwardedNodeForDrag(FGuid& InOutGuid) const { return false; }
-
-protected:
-
-	TArray<TSharedPtr<FSchematicGraphNode>> Nodes;
-	TMap<FGuid, TSharedPtr<FSchematicGraphNode>> NodeByGuid;
-	TArray<TSharedPtr<FSchematicGraphLink>> Links;
-	TMap<FGuid, TSharedPtr<FSchematicGraphLink>> LinkByGuid;
-	TMap<uint32, TSharedPtr<FSchematicGraphLink>> LinkByHash;
-	TMap<FGuid, TTuple<TArray<FGuid>, TArray<FGuid>>> NodeGuidToLinkGuids;
-
-	FOnNodeAdded OnNodeAddedDelegate;
-	FOnNodeRemoved OnNodeRemovedDelegate;
-	FOnLinkAdded OnLinkAddedDelegate;
-	FOnLinkRemoved OnLinkRemovedDelegate;
-	FOnGraphReset OnGraphResetDelegate;
-};
-
 class SSchematicGraphNode;
 
 class FSchematicGraphNodeDragDropOp : public FDragDropOperation
@@ -508,8 +69,8 @@ public:
 	SLATE_ARGUMENT(TSharedPtr<FVector2dAttribute>, Size)
 	SLATE_ARGUMENT(TSharedPtr<FFloatAttribute>, Scale)
 	SLATE_ATTRIBUTE(bool, EnableAutoScale)
-	SLATE_ARGUMENT(TSharedPtr<FLinearColorAttribute>, Color)
-	SLATE_ATTRIBUTE(const FSlateBrush*, Brush)
+	SLATE_ARGUMENT(TArray<TSharedPtr<FLinearColorAttribute>>, LayerColors)
+	SLATE_ARGUMENT(TFunction<const FSlateBrush*(const FGuid&, int32)>, BrushGetter)
 	SLATE_EVENT(FOnClicked, OnClicked)
 	SLATE_EVENT(FOnBeginDrag, OnBeginDrag)
 	SLATE_EVENT(FOnEndDrag, OnEndDrag)
@@ -558,8 +119,8 @@ private:
 	TSharedPtr<FFloatAttribute> Scale;
 	TAttribute<bool> EnableAutoScale;
 	TOptional<float> AutoScale;
-	TSharedPtr<FLinearColorAttribute> Color;
-	TAttribute<const FSlateBrush*> Brush;
+	TArray<TSharedPtr<FLinearColorAttribute>> LayerColors;
+	TFunction<const FSlateBrush*(const FGuid&, int32)> BrushGetter;
 	FOnClicked OnClickedDelegate;
 	FOnBeginDrag OnBeginDragDelegate;
 	FOnEndDrag OnEndDragDelegate;
@@ -647,23 +208,9 @@ public:
 	FReply HandleNodeDragDetected(FGuid Guid, const FGeometry& MyGeometry, const FPointerEvent& MouseEvent);
 
 	virtual FVector2d GetPositionForNode(FGuid InNodeGuid) const;
-	virtual FVector2d GetSizeForNode(FGuid InNodeGuid) const;
 	virtual float GetScaleForNode(FGuid InNodeGuid, bool bIncludeScaleOffset) const;
 	virtual bool IsAutoScaleEnabledForNode(FGuid InNodeGuid) const;
 	virtual float GetMinimumLinkDistanceForNode(FGuid InLinkGuid, bool bIncludeScale = true) const;
-	virtual FLinearColor GetColorForNode(FGuid InNodeGuid) const;
-	virtual const FSlateBrush* GetBrushForNode(FGuid InNodeGuid) const;
-	virtual FText GetToolTipForNode(FGuid InNodeGuid) const;
-	virtual ESchematicGraphVisibility GetVisibilityForNode(FGuid InNodeGuid) const;
-	virtual bool IsDragSupportedForNode(FGuid InNodeGuid) const;
-
-	virtual float GetMinimumForLink(FGuid InLinkGuid) const;
-	virtual float GetMaximumForLink(FGuid InLinkGuid) const;
-	virtual FLinearColor GetColorForLink(FGuid InLinkGuid) const;
-	virtual float GetThicknessForLink(FGuid InLinkGuid) const;
-	virtual const FSlateBrush* GetBrushForLink(FGuid InLinkGuid) const;
-	virtual FText GetToolTipForLink(FGuid InLinkGuid) const;
-	virtual ESchematicGraphVisibility GetVisibilityForLink(FGuid InLinkGuid) const;
 
 	void UpdateAutoScalingForNodes();
 
diff --git a/Engine/Source/Editor/AnimationEditorWidgets/Public/SchematicGraphPanel/SchematicGraphDefines.h b/Engine/Source/Editor/AnimationEditorWidgets/Public/SchematicGraphPanel/SchematicGraphDefines.h
new file mode 100644
--- /dev/null
+++ b/Engine/Source/Editor/AnimationEditorWidgets/Public/SchematicGraphPanel/SchematicGraphDefines.h
@@ -0,0 +1,135 @@
+// Copyright Epic Games, Inc. All Rights Reserved.
+#pragma once
+
+#if WITH_EDITOR
+
+#include "CoreMinimal.h"
+#include "Styling/SlateBrush.h"
+
+class FSchematicGraphTag;
+class FSchematicGraphNode;
+class FSchematicGraphLink;
+class FSchematicGraphModel;
+
+namespace ESchematicGraphPlacementConstraint
+{
+	enum Type : int
+	{
+		Free,
+		TopLeft,
+		TopRight,
+		BottomLeft,
+		BottomRight
+	};
+}
+
+namespace ESchematicGraphVisibility
+{
+	enum Type : int
+	{
+		Visible,
+		FadedOut,
+		Hidden
+	};
+}
+
+DECLARE_EVENT_TwoParams(FSchematicGraph, FOnSchematicGraphTagAdded, const FSchematicGraphNode*, const FSchematicGraphTag*);
+DECLARE_EVENT_TwoParams(FSchematicGraph, FOnTagRemoved, const FSchematicGraphNode*, const FSchematicGraphTag*);
+DECLARE_EVENT_OneParam(FSchematicGraph, FOnSchematicGraphNodeAdded, const FSchematicGraphNode*);
+DECLARE_EVENT_OneParam(FSchematicGraph, FOnNodeRemoved, const FSchematicGraphNode*);
+DECLARE_EVENT_OneParam(FSchematicGraph, FOnSchematicGraphLinkAdded, const FSchematicGraphLink*);
+DECLARE_EVENT_OneParam(FSchematicGraph, FOnLinkRemoved, const FSchematicGraphLink*);
+DECLARE_EVENT(FSchematicGraph, FOnGraphReset);
+
+#define SCHEMATICGRAPHELEMENT_BODY_BASE(ClassName) \
+virtual ~ClassName() {} \
+inline static const FName& Type = TEXT(#ClassName); \
+virtual const FName& GetType() const \
+{ \
+	return ClassName::Type;  \
+} \
+virtual bool IsA(const FName& InType) const \
+{ \
+	return Type == InType; \
+} \
+template<typename ElementType> \
+bool IsA() const \
+{ \
+	return IsA(ElementType::Type); \
+} \
+\
+template<typename T> \
+friend const T* Cast(const ClassName* InElement) \
+{ \
+	if(InElement) \
+	{ \
+		if(InElement->IsA<T>()) \
+		{ \
+			return static_cast<const T*>(InElement); \
+		} \
+	} \
+	return nullptr; \
+} \
+\
+template<typename T> \
+friend T* Cast(ClassName* InElement) \
+{ \
+	if(InElement) \
+	{ \
+		if(InElement->IsA<T>()) \
+		{ \
+			return static_cast<T*>(InElement); \
+		} \
+	} \
+	return nullptr; \
+} \
+\
+template<typename T> \
+friend const T* CastChecked(const ClassName* InElement) \
+{ \
+	const T* Element = Cast<T>(InElement); \
+	check(Element); \
+	return Element; \
+} \
+\
+template<typename T> \
+friend T* CastChecked(ClassName* InElement) \
+{ \
+	T* Element = Cast<T>(InElement); \
+	check(Element); \
+	return Element; \
+}
+
+#define SCHEMATICGRAPHELEMENT_BODY(ClassName, SuperClass, BaseClass) \
+inline static const FName& Type = TEXT(#ClassName); \
+virtual const FName& GetType() const override { return ClassName::Type; } \
+virtual bool IsA(const FName& InType) const \
+{ \
+	if(ClassName::Type == InType) \
+	{ \
+		return true; \
+	} \
+	return SuperClass::IsA(InType); \
+} \
+template<typename T> \
+friend const T* Cast(const ClassName* InElement) \
+{ \
+	return Cast<T>((const BaseClass*) InElement); \
+} \
+template<typename T> \
+friend T* Cast(ClassName* InElement) \
+{ \
+	return Cast<T>((BaseClass*) InElement); \
+} \
+template<typename T> \
+friend const T* CastChecked(const ClassName* InElement) \
+{ \
+	return CastChecked<T>((const BaseClass*) InElement); \
+} \
+template<typename T> \
+friend T* CastChecked(ClassName* InElement) \
+{ \
+	return CastChecked<T>((BaseClass*) InElement); \
+}
+
+#endif
\ No newline at end of file
diff --git a/Engine/Source/Editor/AnimationEditorWidgets/Public/SchematicGraphPanel/SchematicGraphLink.h b/Engine/Source/Editor/AnimationEditorWidgets/Public/SchematicGraphPanel/SchematicGraphLink.h
new file mode 100644
--- /dev/null
+++ b/Engine/Source/Editor/AnimationEditorWidgets/Public/SchematicGraphPanel/SchematicGraphLink.h
@@ -0,0 +1,55 @@
+// Copyright Epic Games, Inc. All Rights Reserved.
+#pragma once
+
+#if WITH_EDITOR
+
+#include "SchematicGraphDefines.h"
+
+#define SCHEMATICGRAPHLINK_BODY(ClassName, SuperClass) \
+SCHEMATICGRAPHELEMENT_BODY(ClassName, SuperClass, FSchematicGraphNode)
+
+class ANIMATIONEDITORWIDGETS_API FSchematicGraphLink
+{
+public:
+
+	SCHEMATICGRAPHELEMENT_BODY_BASE(FSchematicGraphLink)
+	
+	const FGuid& GetGuid() const { return Guid; }
+	static uint32 GetLinkHash(const FGuid& InSourceNodeGuid, const FGuid& InTargetNodeGuid)
+	{
+		return HashCombine(GetTypeHash(InSourceNodeGuid), GetTypeHash(InTargetNodeGuid));
+	}
+	uint32 GetLinkHash() const { return GetLinkHash(SourceNodeGuid, TargetNodeGuid); }
+	const FGuid& GetSourceNodeGuid() const { return SourceNodeGuid; }
+	const FGuid& GetTargetNodeGuid() const { return TargetNodeGuid; }
+	virtual float GetMinimum() const { return Minimum; }
+	virtual float GetMaximum() const { return Maximum; }
+	virtual FVector2d GetSourceNodeOffset() const { return SourceNodeOffset; }
+	virtual FVector2d GetTargetNodeOffset() const { return TargetNodeOffset; }
+	virtual FLinearColor GetColor() const { return Color; }
+	virtual float GetThickness() const { return Thickness; }
+	virtual const FSlateBrush* GetBrush() const { return Brush; }
+	virtual const FText& GetToolTip() const { return ToolTip; }
+	virtual ESchematicGraphVisibility::Type GetVisibility() const { return Visibility; }
+	virtual void SetVisibility(ESchematicGraphVisibility::Type InVisibility) { Visibility = InVisibility; }
+
+protected:
+	
+	FSchematicGraphModel* Model = nullptr;
+	FGuid Guid = FGuid::NewGuid();
+	FGuid SourceNodeGuid = FGuid();
+	FGuid TargetNodeGuid = FGuid();
+	float Minimum = 0.f;
+	float Maximum = 1.f;
+	FVector2d SourceNodeOffset = FVector2d::ZeroVector;
+	FVector2d TargetNodeOffset = FVector2d::ZeroVector;
+	FLinearColor Color = FLinearColor::White;
+	float Thickness = 1.f;
+	const FSlateBrush* Brush = nullptr;
+	FText ToolTip = FText();
+	ESchematicGraphVisibility::Type Visibility = ESchematicGraphVisibility::Visible;
+
+	friend class FSchematicGraphModel;
+};
+
+#endif
\ No newline at end of file
diff --git a/Engine/Source/Editor/AnimationEditorWidgets/Public/SchematicGraphPanel/SchematicGraphModel.h b/Engine/Source/Editor/AnimationEditorWidgets/Public/SchematicGraphPanel/SchematicGraphModel.h
new file mode 100644
--- /dev/null
+++ b/Engine/Source/Editor/AnimationEditorWidgets/Public/SchematicGraphPanel/SchematicGraphModel.h
@@ -0,0 +1,235 @@
+// Copyright Epic Games, Inc. All Rights Reserved.
+#pragma once
+
+#if WITH_EDITOR
+
+#include "SchematicGraphTag.h"
+#include "SchematicGraphNode.h"
+#include "SchematicGraphLink.h"
+
+class ANIMATIONEDITORWIDGETS_API FSchematicGraphModel
+{
+public:
+
+	virtual ~FSchematicGraphModel() {}
+	
+	virtual void Reset();
+	
+	template<typename NodeType = FSchematicGraphNode>
+	NodeType* AddNode(bool bNotify = true)
+	{
+		const TSharedPtr<FSchematicGraphNode> NewNode = MakeShareable(new NodeType);
+		NewNode->Model = this;
+		Nodes.Add(NewNode);
+		NodeByGuid.Add(NewNode->GetGuid(), NewNode);
+
+		if (bNotify && OnNodeAddedDelegate.IsBound())
+		{
+			OnNodeAddedDelegate.Broadcast(NewNode.Get());
+		}
+		return static_cast<NodeType*>(NewNode.Get());
+	}
+
+	template<typename NodeType = FSchematicGraphNode>
+	const NodeType* FindNode(const FGuid& InNodeGuid) const
+	{
+		if(const TSharedPtr<FSchematicGraphNode>* ExistingNode = NodeByGuid.Find(InNodeGuid))
+		{
+			return Cast<NodeType>(ExistingNode->Get());
+		};
+		return nullptr;
+	}
+
+	template<typename NodeType = FSchematicGraphNode>
+	const NodeType* FindNodeChecked(const FGuid& InNodeGuid) const
+	{
+		if(const TSharedPtr<FSchematicGraphNode>* ExistingNode = NodeByGuid.Find(InNodeGuid))
+		{
+			check(ExistingNode->Get()->IsA(NodeType::Type));
+			return static_cast<NodeType*>(ExistingNode->Get());
+		};
+		return nullptr;
+	}
+
+	virtual bool RemoveNode(const FGuid& InNodeGuid);
+
+	virtual bool SetParentNode(const FGuid& InChildNodeGuid, const FGuid& InParentNodeGuid);
+	virtual bool SetParentNode(const FSchematicGraphNode* InChildNode, const FSchematicGraphNode* InParentNode);
+	bool RemoveFromParentNode(const FGuid& InChildNodeGuid);
+	bool RemoveFromParentNode(const FSchematicGraphNode* InChildNode);
+
+	const TArray<TSharedPtr<FSchematicGraphNode>>& GetNodes() const { return Nodes; }
+
+	FVector2d GetPositionForNode(const FGuid& InNodeGuid) const;
+	virtual FVector2d GetPositionForNode(const FSchematicGraphNode* InNode) const;
+	FVector2d GetPositionOffsetForNode(const FGuid& InNodeGuid) const;
+	virtual FVector2d GetPositionOffsetForNode(const FSchematicGraphNode* InNode) const;
+	bool GetPositionAnimationEnabledForNode(const FGuid& InNodeGuid) const;
+	virtual bool GetPositionAnimationEnabledForNode(const FSchematicGraphNode* InNode) const;
+	FVector2d GetSizeForNode(const FGuid& InNodeGuid) const;
+	virtual FVector2d GetSizeForNode(const FSchematicGraphNode* InNode) const;
+	float GetScaleForNode(const FGuid& InNodeGuid, bool bIncludeScaleOffset) const;
+	virtual float GetScaleForNode(const FSchematicGraphNode* InNode, bool bIncludeScaleOffset) const;
+	float GetScaleOffsetForNode(const FGuid& InNodeGuid) const;
+	virtual float GetScaleOffsetForNode(const FSchematicGraphNode* InNode) const;
+	float GetMinimumLinkDistanceForNode(const FGuid& InNodeGuid) const;
+	virtual float GetMinimumLinkDistanceForNode(const FSchematicGraphNode* InNode) const;
+	bool IsAutoScaleEnabledForNode(const FGuid& InNodeGuid) const;
+	virtual bool IsAutoScaleEnabledForNode(const FSchematicGraphNode* InNode) const;
+	int32 GetNumLayersForNode(const FGuid& InNodeGuid) const;
+	virtual int32 GetNumLayersForNode(const FSchematicGraphNode* InNode) const;
+	FLinearColor GetColorForNode(const FGuid& InNodeGuid, int32 InLayerIndex) const;
+	virtual FLinearColor GetColorForNode(const FSchematicGraphNode* InNode, int32 InLayerIndex) const;
+	const FSlateBrush* GetBrushForNode(const FGuid& InNodeGuid, int32 InLayerIndex) const;
+	virtual const FSlateBrush* GetBrushForNode(const FSchematicGraphNode* InNode, int32 InLayerIndex) const;
+	FText GetToolTipForNode(const FGuid& InNodeGuid) const;
+	virtual FText GetToolTipForNode(const FSchematicGraphNode* InNode) const;
+	ESchematicGraphPlacementConstraint::Type GetPlacementForNode(const FGuid& InNodeGuid) const;
+	virtual ESchematicGraphPlacementConstraint::Type GetPlacementForNode(const FSchematicGraphNode* InNode) const;
+	ESchematicGraphVisibility::Type GetVisibilityForNode(const FGuid& InNodeGuid) const;
+	virtual ESchematicGraphVisibility::Type GetVisibilityForNode(const FSchematicGraphNode* InNode) const;
+	ESchematicGraphVisibility::Type GetVisibilityForChildNodes(const FGuid& InNodeGuid) const;
+	virtual ESchematicGraphVisibility::Type GetVisibilityForChildNodes(const FSchematicGraphNode* InNode) const;
+	bool IsDragSupportedForNode(const FGuid& InNodeGuid) const;
+	virtual bool IsDragSupportedForNode(const FSchematicGraphNode* InNode) const;
+
+	FLinearColor GetBackgroundColorForTag(const FGuid& InNodeGuid, const FGuid& InTagGuid) const;
+	virtual FLinearColor GetBackgroundColorForTag(const FSchematicGraphTag* InTag) const;
+	FLinearColor GetForegroundColorForTag(const FGuid& InNodeGuid, const FGuid& InTagGuid) const;
+	virtual FLinearColor GetForegroundColorForTag(const FSchematicGraphTag* InTag) const;
+	FLinearColor GetLabelColorForTag(const FGuid& InNodeGuid, const FGuid& InTagGuid) const;
+	virtual FLinearColor GetLabelColorForTag(const FSchematicGraphTag* InTag) const;
+	const FSlateBrush* GetBackgroundBrushForTag(const FGuid& InNodeGuid, const FGuid& InTagGuid) const;
+	virtual const FSlateBrush* GetBackgroundBrushForTag(const FSchematicGraphTag* InTag) const;
+	const FSlateBrush* GetForegroundBrushForTag(const FGuid& InNodeGuid, const FGuid& InTagGuid) const;
+	virtual const FSlateBrush* GetForegroundBrushForTag(const FSchematicGraphTag* InTag) const;
+	const FText GetLabelForTag(const FGuid& InNodeGuid, const FGuid& InTagGuid) const;
+	virtual const FText GetLabelForTag(const FSchematicGraphTag* InTag) const;
+	const FText GetToolTipForTag(const FGuid& InNodeGuid, const FGuid& InTagGuid) const;
+	virtual const FText GetToolTipForTag(const FSchematicGraphTag* InTag) const;
+	ESchematicGraphVisibility::Type GetVisibilityForTag(const FGuid& InNodeGuid, const FGuid& InTagGuid) const;
+	virtual ESchematicGraphVisibility::Type GetVisibilityForTag(const FSchematicGraphTag* InTag) const;
+
+	template<typename LinkType = FSchematicGraphLink>
+	LinkType* AddLink(const FGuid& InSourceNodeGuid, const FGuid& InTargetNodeGuid, bool bNotify = true)
+	{
+		const TSharedPtr<FSchematicGraphLink> NewLink = MakeShareable(new LinkType);
+		NewLink->Model = this;
+		NewLink->SourceNodeGuid = InSourceNodeGuid;
+		NewLink->TargetNodeGuid = InTargetNodeGuid;
+		Links.Add(NewLink);
+		LinkByGuid.Add(NewLink->GetGuid(), NewLink);
+		LinkByHash.Add(NewLink->GetLinkHash(), NewLink);
+		NodeGuidToLinkGuids.FindOrAdd(InSourceNodeGuid).Get<0>().Add(NewLink->GetGuid());
+		NodeGuidToLinkGuids.FindOrAdd(InTargetNodeGuid).Get<1>().Add(NewLink->GetGuid());
+
+		if (bNotify && OnLinkAddedDelegate.IsBound())
+		{
+			OnLinkAddedDelegate.Broadcast(NewLink.Get());
+		}
+		return static_cast<LinkType*>(NewLink.Get());
+	}
+
+	template<typename LinkType = FSchematicGraphLink>
+	const LinkType* FindLink(const FGuid& InLinkGuid) const
+	{
+		if(const TSharedPtr<FSchematicGraphLink>* ExistingLink = LinkByGuid.Find(InLinkGuid))
+		{
+			return Cast<LinkType>(ExistingLink->Get());
+		};
+		return nullptr;
+	}
+
+	template<typename LinkType = FSchematicGraphLink>
+	const LinkType* FindLinkChecked(const FGuid& InLinkGuid) const
+	{
+		if(const TSharedPtr<FSchematicGraphLink>* ExistingLink = LinkByGuid.Find(InLinkGuid))
+		{
+			check(ExistingLink->Get()->IsA(LinkType::Type));
+			return static_cast<LinkType*>(ExistingLink->Get());
+		};
+		return nullptr;
+	}
+
+	template<typename LinkType = FSchematicGraphLink>
+	const LinkType* FindLink(const FGuid& InSourceNodeGuid, const FGuid& InTargetNodeGuid) const
+	{
+		const uint32 LinkHash = FSchematicGraphLink::GetLinkHash(InSourceNodeGuid, InTargetNodeGuid);
+		if(const TSharedPtr<FSchematicGraphLink>* ExistingLink = LinkByHash.Find(LinkHash))
+		{
+			return Cast<LinkType>(ExistingLink->Get());
+		};
+		return nullptr;
+	}
+
+	template<typename LinkType = FSchematicGraphLink>
+	const LinkType* FindLinkChecked(const FGuid& InSourceNodeGuid, const FGuid& InTargetNodeGuid) const
+	{
+		const uint32 LinkHash = FSchematicGraphLink::GetLinkHash(InSourceNodeGuid, InTargetNodeGuid);
+		if(const TSharedPtr<FSchematicGraphLink>* ExistingLink = LinkByHash.Find(LinkHash))
+		{
+			check(ExistingLink->Get()->IsA(LinkType::Type));
+			return Cast<LinkType>(ExistingLink->Get());
+		};
+		return nullptr;
+	}
+
+	bool IsLinkedTo(const FGuid& InSourceNodeGuid, const FGuid& InTargetNodeGuid) const;
+	TArray<const FSchematicGraphLink*> FindLinksOnNode(const FGuid& InNodeGuid) const;
+	TArray<const FSchematicGraphLink*> FindLinksOnSource(const FGuid& InSourceNodeGuid) const;
+	TArray<const FSchematicGraphLink*> FindLinksOnTarget(const FGuid& InTargetNodeGuid) const;
+
+	virtual bool RemoveLink(const FGuid& InLinkGuid);
+
+	const TArray<TSharedPtr<FSchematicGraphLink>>& GetLinks() const { return Links; }
+
+	float GetMinimumForLink(const FGuid& InLinkGuid) const;
+	virtual float GetMinimumForLink(const FSchematicGraphLink* InLink) const;
+	float GetMaximumForLink(const FGuid& InLinkGuid) const;
+	virtual float GetMaximumForLink(const FSchematicGraphLink* InLink) const;
+	FVector2d GetSourceNodeOffsetForLink(const FGuid& InLinkGuid) const;
+	FVector2d GetSourceNodeOffsetForLink(const FSchematicGraphLink* InLink) const;
+	FVector2d GetTargetNodeOffsetForLink(const FGuid& InLinkGuid) const;
+	FVector2d GetTargetNodeOffsetForLink(const FSchematicGraphLink* InLink) const;
+	FLinearColor GetColorForLink(const FGuid& InLinkGuid) const;
+	virtual FLinearColor GetColorForLink(const FSchematicGraphLink* InLink) const;
+	float GetThicknessForLink(const FGuid& InLinkGuid) const;
+	virtual float GetThicknessForLink(const FSchematicGraphLink* InLink) const;
+	const FSlateBrush* GetBrushForLink(const FGuid& InLinkGuid) const;
+	virtual const FSlateBrush* GetBrushForLink(const FSchematicGraphLink* InLink) const;
+	const FText GetToolTipForLink(const FGuid& InLinkGuid) const;
+	virtual const FText GetToolTipForLink(const FSchematicGraphLink* InLink) const;
+	ESchematicGraphVisibility::Type GetVisibilityForLink(const FGuid& InLinkGuid) const;
+	virtual ESchematicGraphVisibility::Type GetVisibilityForLink(const FSchematicGraphLink* InLink) const;
+
+	FOnSchematicGraphNodeAdded& OnNodeAdded() { return OnNodeAddedDelegate; }
+	FOnNodeRemoved& OnNodeRemoved() { return OnNodeRemovedDelegate; }
+	FOnSchematicGraphLinkAdded& OnLinkAdded() { return OnLinkAddedDelegate; }
+	FOnLinkRemoved& OnLinkRemoved() { return OnLinkRemovedDelegate; }
+	FOnSchematicGraphTagAdded& OnTagAdded() { return OnTagAddedDelegate; }
+	FOnTagRemoved& OnTagRemoved() { return OnTagRemovedDelegate; }
+	FOnGraphReset& OnGraphReset() { return OnGraphResetDelegate; }
+
+	virtual bool GetForwardedNodeForDrag(FGuid& InOutGuid) const { return false; }
+
+protected:
+
+	TArray<TSharedPtr<FSchematicGraphNode>> Nodes;
+	TMap<FGuid, TSharedPtr<FSchematicGraphNode>> NodeByGuid;
+	TArray<TSharedPtr<FSchematicGraphLink>> Links;
+	TMap<FGuid, TSharedPtr<FSchematicGraphLink>> LinkByGuid;
+	TMap<uint32, TSharedPtr<FSchematicGraphLink>> LinkByHash;
+	TMap<FGuid, TTuple<TArray<FGuid>, TArray<FGuid>>> NodeGuidToLinkGuids;
+
+	FOnSchematicGraphNodeAdded OnNodeAddedDelegate;
+	FOnNodeRemoved OnNodeRemovedDelegate;
+	FOnSchematicGraphLinkAdded OnLinkAddedDelegate;
+	FOnLinkRemoved OnLinkRemovedDelegate;
+	FOnSchematicGraphTagAdded OnTagAddedDelegate;
+	FOnTagRemoved OnTagRemovedDelegate;
+	FOnGraphReset OnGraphResetDelegate;
+
+	friend class FSchematicGraphNode;
+};
+
+#endif
\ No newline at end of file
diff --git a/Engine/Source/Editor/AnimationEditorWidgets/Public/SchematicGraphPanel/SchematicGraphNode.h b/Engine/Source/Editor/AnimationEditorWidgets/Public/SchematicGraphPanel/SchematicGraphNode.h
new file mode 100644
--- /dev/null
+++ b/Engine/Source/Editor/AnimationEditorWidgets/Public/SchematicGraphPanel/SchematicGraphNode.h
@@ -0,0 +1,129 @@
+// Copyright Epic Games, Inc. All Rights Reserved.
+#pragma once
+
+#if WITH_EDITOR
+
+#include "SchematicGraphDefines.h"
+
+#define SCHEMATICGRAPHNODE_BODY(ClassName, SuperClass) \
+SCHEMATICGRAPHELEMENT_BODY(ClassName, SuperClass, FSchematicGraphNode)
+
+class ANIMATIONEDITORWIDGETS_API FSchematicGraphNode
+{
+public:
+
+	SCHEMATICGRAPHELEMENT_BODY_BASE(FSchematicGraphNode)
+	
+	const FGuid& GetGuid() const { return Guid; }
+	bool HasParentNode() const { return ParentNodeGuid.IsValid(); }
+	const FGuid& GetParentNodeGuid() const { return ParentNodeGuid; }
+	const FSchematicGraphNode* GetParentNode() const;
+	int32 GetNumChildNodes() const { return ChildNodeGuids.Num(); }
+	const TArray<FGuid>& GetChildNodeGuids() const { return ChildNodeGuids; }
+	const FSchematicGraphNode* GetChildNode(int32 InChildNodeIndex) const;
+	virtual ESchematicGraphVisibility::Type GetVisibilityForChildNodes() const { return GetVisibility(); }
+	bool IsSelected() const { return bIsSelected; }
+	void SetSelected(bool bSelected = true) { bIsSelected = bSelected; }
+	virtual FVector2d GetPosition() const { return Position; }
+	virtual FVector2d GetPositionOffset() const { return PositionOffset; }
+	virtual void SetPositionOffset(const FVector2d& InPositionOffset) { PositionOffset = InPositionOffset; }
+	virtual float GetScaleOffset() const { return ScaleOffset; }
+	virtual void SetScaleOffset(float InScaleOffset) { ScaleOffset = InScaleOffset; }
+	int32 GetNumLayers() const { return FMath::Min(Colors.Num(), Brushes.Num()); }
+	virtual FLinearColor GetColor(int32 InLayerIndex) const { return Colors[InLayerIndex]; }
+	virtual const FSlateBrush* GetBrush(int32 InLayerIndex) const { return Brushes[InLayerIndex]; }
+	virtual bool IsAutoScaleEnabled() const { return false; }
+	virtual const FText& GetToolTip() const { return ToolTip; }
+	virtual ESchematicGraphPlacementConstraint::Type GetPlacement() const { return Placement; }
+	virtual void SetPlacement(ESchematicGraphPlacementConstraint::Type InPlacement) { Placement = InPlacement; }
+	virtual ESchematicGraphVisibility::Type GetVisibility() const;
+	virtual void SetVisibility(ESchematicGraphVisibility::Type InVisibility) { Visibility = InVisibility; }
+	virtual bool IsDragSupported() const { return bDragSupported; }
+	virtual void SetDragSupported(bool InDragSupported) { bDragSupported = InDragSupported;}
+
+	virtual FString GetDragDropDecoratorLabel() const;
+
+	template<typename TagType = FSchematicGraphTag>
+	TagType* AddTag(bool bNotify = true)
+	{
+		const TSharedPtr<FSchematicGraphTag> NewTag = MakeShareable(new TagType);
+		NewTag->Model = Model;
+		Tags.Add(NewTag);
+		TagByGuid.Add(NewTag->GetGuid(), NewTag);
+
+		if (bNotify)
+		{
+			NotifyTagAdded(NewTag);
+		}
+		return static_cast<TagType*>(NewTag.Get());
+	}
+
+	template<typename TagType = FSchematicGraphTag>
+	const TagType* FindTag(const FGuid& InTagGuid) const
+	{
+		if(const TSharedPtr<FSchematicGraphTag>* ExistingTag = TagByGuid.Find(InTagGuid))
+		{
+			return Cast<TagType>(ExistingTag->Get());
+		};
+		return nullptr;
+	}
+
+	template<typename TagType = FSchematicGraphTag>
+	const TagType* FindTagChecked(const FGuid& InTagGuid) const
+	{
+		if(const TSharedPtr<FSchematicGraphTag>* ExistingTag = TagByGuid.Find(InTagGuid))
+		{
+			check(ExistingTag->Get()->IsA(TagType::Type));
+			return static_cast<TagType*>(ExistingTag->Get());
+		};
+		return nullptr;
+	}
+
+	virtual bool RemoveTag(const FGuid& InTagGuid);
+
+	const TArray<TSharedPtr<FSchematicGraphTag>>& GetTags() const { return Tags; }
+
+protected:
+
+	void NotifyTagAdded(const TSharedPtr<FSchematicGraphTag>& Tag);
+
+	FSchematicGraphModel* Model = nullptr;
+	FGuid Guid = FGuid::NewGuid();
+	FGuid ParentNodeGuid = FGuid();
+	TArray<FGuid> ChildNodeGuids; 
+	bool bIsSelected = false;
+	FVector2d Position = FVector2d::ZeroVector;
+	FVector2d PositionOffset = FVector2d::ZeroVector;
+	float ScaleOffset = 1.f;
+	TArray<FLinearColor> Colors = { FLinearColor::White };
+	TArray<const FSlateBrush*> Brushes = { nullptr };
+	FText ToolTip = FText();
+	ESchematicGraphPlacementConstraint::Type Placement = ESchematicGraphPlacementConstraint::Free;
+	ESchematicGraphVisibility::Type Visibility = ESchematicGraphVisibility::Visible;
+	bool bDragSupported = false;
+
+	TArray<TSharedPtr<FSchematicGraphTag>> Tags;
+	TMap<FGuid, TSharedPtr<FSchematicGraphTag>> TagByGuid;
+
+	friend class FSchematicGraphModel;
+};
+
+class FSchematicGraphGroupNode : public FSchematicGraphNode
+{
+public:
+
+	SCHEMATICGRAPHNODE_BODY(FSchematicGraphGroupNode, FSchematicGraphNode)
+
+	virtual ~FSchematicGraphGroupNode() override {}
+	bool IsExpanded() const { return bExpanded; }
+	void SetExpanded(bool InExpanded) { bExpanded = InExpanded; }
+	virtual ESchematicGraphVisibility::Type GetVisibilityForChildNodes() const override;
+	
+protected:
+
+	bool bExpanded = false;
+
+	friend class FControlRigSchematicModel;
+};
+
+#endif
\ No newline at end of file
diff --git a/Engine/Source/Editor/AnimationEditorWidgets/Public/SchematicGraphPanel/SchematicGraphTag.h b/Engine/Source/Editor/AnimationEditorWidgets/Public/SchematicGraphPanel/SchematicGraphTag.h
new file mode 100644
--- /dev/null
+++ b/Engine/Source/Editor/AnimationEditorWidgets/Public/SchematicGraphPanel/SchematicGraphTag.h
@@ -0,0 +1,57 @@
+// Copyright Epic Games, Inc. All Rights Reserved.
+#pragma once
+
+#if WITH_EDITOR
+
+#include "SchematicGraphDefines.h"
+
+#define SCHEMATICGRAPHTAG_BODY(ClassName, SuperClass) \
+SCHEMATICGRAPHELEMENT_BODY(ClassName, SuperClass, FSchematicGraphTag)
+
+class ANIMATIONEDITORWIDGETS_API FSchematicGraphTag
+{
+public:
+	
+	SCHEMATICGRAPHELEMENT_BODY_BASE(FSchematicGraphTag)
+
+	const FSchematicGraphNode* GetNode() const;
+	const FGuid& GetGuid() const { return Guid; }
+	virtual FLinearColor GetBackgroundColor() const { return BackgroundColor; }
+	virtual void SetBackgroundColor(const FLinearColor& InBackgroundColor) { BackgroundColor = InBackgroundColor; }
+	virtual FLinearColor GetForegroundColor() const { return ForegroundColor; }
+	virtual void SetForegroundColor(const FLinearColor& InForegroundColor) { ForegroundColor = InForegroundColor; }
+	virtual FLinearColor GetLabelColor() const { return LabelColor; }
+	virtual void SetLabelColor(const FLinearColor& InLabelColor) { LabelColor = InLabelColor; }
+	virtual const FSlateBrush* GetBackgroundBrush() const { return BackgroundBrush; }
+	virtual void SetBackgroundBrush(const FSlateBrush* InBackgroundBrush) { BackgroundBrush = InBackgroundBrush; }
+	virtual const FSlateBrush* GetForegroundBrush() const { return ForegroundBrush; }
+	virtual void SetForegroundBrush(const FSlateBrush* InForegroundBrush) { ForegroundBrush = InForegroundBrush; }
+	virtual const FText& GetLabel() const { return Label; }
+	virtual void SetLabel(const FText& InLabel) { Label = InLabel; }
+	virtual const FText& GetToolTip() const { return ToolTip; }
+	virtual void SetToolTip(const FText& InToolTip) { ToolTip = InToolTip; }
+	virtual float GetPlacementAngle() const { return PlacementAngle; }
+	virtual void SetPlacementAngle(float InPlacementAngle) { PlacementAngle = InPlacementAngle; }
+	virtual ESchematicGraphVisibility::Type GetVisibility() const { return Visibility; }
+	virtual void SetVisibility(ESchematicGraphVisibility::Type InVisibility) { Visibility = InVisibility; }
+
+protected:
+
+	FSchematicGraphModel* Model = nullptr;
+	TWeakPtr<FSchematicGraphNode> Node;
+	FGuid Guid = FGuid::NewGuid();
+	FLinearColor BackgroundColor = FLinearColor(FColor(0, 112, 224));
+	FLinearColor ForegroundColor = FLinearColor::White;
+	FLinearColor LabelColor = FLinearColor::White;
+	const FSlateBrush* BackgroundBrush = nullptr;
+	const FSlateBrush* ForegroundBrush = nullptr;
+	FText Label = FText();
+	FText ToolTip = FText();
+	float PlacementAngle = -45.f;
+	ESchematicGraphVisibility::Type Visibility = ESchematicGraphVisibility::Visible;
+
+	friend class FSchematicGraphModel;
+	friend class FSchematicGraphNode;
+};
+
+#endif
\ No newline at end of file
