Schematic View: Fix missing include

#rnx

diff --git a/Engine/Source/Editor/AnimationEditorWidgets/Public/SchematicGraphPanel/SchematicGraphNode.h b/Engine/Source/Editor/AnimationEditorWidgets/Public/SchematicGraphPanel/SchematicGraphNode.h
--- a/Engine/Source/Editor/AnimationEditorWidgets/Public/SchematicGraphPanel/SchematicGraphNode.h
+++ b/Engine/Source/Editor/AnimationEditorWidgets/Public/SchematicGraphPanel/SchematicGraphNode.h
@@ -3,7 +3,7 @@
 
 #if WITH_EDITOR
 
-#include "SchematicGraphDefines.h"
+#include "SchematicGraphTag.h"
 
 #define SCHEMATICGRAPHNODE_BODY(ClassName, SuperClass) \
 SCHEMATICGRAPHELEMENT_BODY(ClassName, SuperClass, FSchematicGraphNode)
