SchematicGraph: Implement auto grouping nodes

#rnx

diff --git a/Engine/Plugins/Animation/ControlRig/Source/ControlRig/Private/ModularRigController.cpp b/Engine/Plugins/Animation/ControlRig/Source/ControlRig/Private/ModularRigController.cpp
--- a/Engine/Plugins/Animation/ControlRig/Source/ControlRig/Private/ModularRigController.cpp
+++ b/Engine/Plugins/Animation/ControlRig/Source/ControlRig/Private/ModularRigController.cpp
@@ -380,6 +380,12 @@ bool UModularRigController::DisconnectConnector(const FRigElementKey& InConnecto
 		return false;
 	}
 
+	const FRigElementKey ConnectorKey(*ConnectorName, ERigElementType::Connector);
+	if(!Module->Connections.Contains(ConnectorKey))
+	{
+		return false;
+	}
+
 #if WITH_EDITOR
 	TSharedPtr<FScopedTransaction> TransactionPtr;
 	if (bSetupUndo)
@@ -389,7 +395,6 @@ bool UModularRigController::DisconnectConnector(const FRigElementKey& InConnecto
 	}
 #endif 
 
-	const FRigElementKey ConnectorKey(*ConnectorName, ERigElementType::Connector);
 	Module->Connections.Remove(ConnectorKey);
 
 	if (Connector->Settings.Type == EConnectorType::Primary)
diff --git a/Engine/Plugins/Animation/ControlRig/Source/ControlRigEditor/Private/ControlRigEditorStyle.h b/Engine/Plugins/Animation/ControlRig/Source/ControlRigEditor/Private/ControlRigEditorStyle.h
--- a/Engine/Plugins/Animation/ControlRig/Source/ControlRigEditor/Private/ControlRigEditorStyle.h
+++ b/Engine/Plugins/Animation/ControlRig/Source/ControlRigEditor/Private/ControlRigEditorStyle.h
@@ -148,10 +148,6 @@ public:
 
 		// Schematic
 		{
-			Set("ControlRig.Schematic.Outline", new IMAGE_BRUSH_SVG("Slate/Outline_Schematic", Icon128x128));
-			Set("ControlRig.Schematic.Background", new IMAGE_BRUSH_SVG("Slate/Background_Schematic", Icon128x128));
-			Set("ControlRig.Schematic.SocketUnresolved", new IMAGE_BRUSH_SVG("Slate/Socket_Open_Schematic", Icon128x128));
-			Set("ControlRig.Schematic.SocketResolved", new IMAGE_BRUSH_SVG("Slate/Socket_Closed_Schematic", Icon128x128));
 			Set("ControlRig.Schematic.ConnectorPrimary", new IMAGE_BRUSH_SVG("Slate/Connector_Primary", Icon128x128));
 			Set("ControlRig.Schematic.ConnectorSecondary", new IMAGE_BRUSH_SVG("Slate/Connector_Secondary", Icon128x128));
 			Set("ControlRig.Schematic.ConnectorOptional", new IMAGE_BRUSH_SVG("Slate/Connector_Optional", Icon128x128));
diff --git a/Engine/Plugins/Animation/ControlRig/Source/ControlRigEditor/Private/Editor/ControlRigEditor.cpp b/Engine/Plugins/Animation/ControlRig/Source/ControlRigEditor/Private/Editor/ControlRigEditor.cpp
--- a/Engine/Plugins/Animation/ControlRig/Source/ControlRigEditor/Private/Editor/ControlRigEditor.cpp
+++ b/Engine/Plugins/Animation/ControlRig/Source/ControlRigEditor/Private/Editor/ControlRigEditor.cpp
@@ -2571,10 +2571,10 @@ void FControlRigEditor::CacheNameLists()
 
 void FControlRigEditor::HandleSchematicViewportCreated(const TSharedRef<SSchematicGraphPanel>& InViewport)
 {
-	InViewport->OnNodeClickedDelegate.BindRaw(&SchematicModel, &FControlRigSchematicModel::HandleSchematicNodeClicked);
-	InViewport->OnBeginDragDelegate.BindRaw(&SchematicModel, &FControlRigSchematicModel::HandleSchematicBeginDrag);
-	InViewport->OnEndDragDelegate.BindRaw(&SchematicModel, &FControlRigSchematicModel::HandleSchematicEndDrag);
-	InViewport->OnDropDelegate.BindRaw(&SchematicModel, &FControlRigSchematicModel::HandleSchematicDrop);
+	InViewport->OnNodeClicked().BindRaw(&SchematicModel, &FControlRigSchematicModel::HandleSchematicNodeClicked);
+	InViewport->OnBeginDrag().BindRaw(&SchematicModel, &FControlRigSchematicModel::HandleSchematicBeginDrag);
+	InViewport->OnEndDrag().BindRaw(&SchematicModel, &FControlRigSchematicModel::HandleSchematicEndDrag);
+	InViewport->OnAcceptDrop().BindRaw(&SchematicModel, &FControlRigSchematicModel::HandleSchematicDrop);
 }
 
 FVector2D FControlRigEditor::ComputePersonaProjectedScreenPos(const FVector& InWorldPos, bool bClampToScreenRectangle)
diff --git a/Engine/Plugins/Animation/ControlRig/Source/ControlRigEditor/Private/Editor/ControlRigSchematicModel.cpp b/Engine/Plugins/Animation/ControlRig/Source/ControlRigEditor/Private/Editor/ControlRigSchematicModel.cpp
--- a/Engine/Plugins/Animation/ControlRig/Source/ControlRigEditor/Private/Editor/ControlRigSchematicModel.cpp
+++ b/Engine/Plugins/Animation/ControlRig/Source/ControlRigEditor/Private/Editor/ControlRigSchematicModel.cpp
@@ -1,6 +1,7 @@
 // Copyright Epic Games, Inc. All Rights Reserved.
 
 #include "ControlRigSchematicModel.h"
+
 #include "ControlRig.h"
 #include "ControlRigBlueprint.h"
 #include "ControlRigEditor.h"
@@ -10,6 +11,9 @@
 #include "DragAndDrop/AssetDragDropOp.h"
 #include "Rigs/RigHierarchyController.h"
 #include "Async/TaskGraphInterfaces.h"
+#include <SchematicGraphPanel/SchematicGraphStyle.h>
+
+#define LOCTEXT_NAMESPACE "ControlRigSchematicModel"
 
 FString FControlRigSchematicRigElementKeyNode::GetDragDropDecoratorLabel() const
 {
@@ -48,7 +52,7 @@ void FControlRigSchematicModel::SetEditor(const TSharedRef<FControlRigEditor>& I
 
 void FControlRigSchematicModel::Reset()
 {
-	FSchematicGraphModel::Reset();
+	Super::Reset();
 	RigElementKeyToGuid.Reset();
 }
 
@@ -59,16 +63,48 @@ FControlRigSchematicRigElementKeyNode* FControlRigSchematicModel::AddElementKeyN
 	{
 		Node->Key = InKey;
 
-		/*
-		FSchematicGraphTag* Tag = Node->AddTag();
-		static const FSlateBrush* BackgroundBrush = FControlRigEditorStyle::Get().GetBrush( "ControlRig.Schematic.Background");
-		Tag->SetBackgroundBrush(BackgroundBrush);
-		Tag->SetLabelColor(FLinearColor::White);
-		Tag->SetLabel(FText::FromString(FString::FromInt(Nodes.Num() * 7)));
-		*/
+		Node->SetLabel(FText::FromString(Node->GetKey().ToString())); 
+		switch(Node->GetKey().Type)
+		{
+			case ERigElementType::Socket:
+			{
+				if (ControlRigBeingDebuggedPtr.IsValid())
+				{
+					if (const URigHierarchy* Hierarchy = ControlRigBeingDebuggedPtr->GetHierarchy())
+					{
+						if(const FRigSocketElement* Socket = Hierarchy->Find<FRigSocketElement>(Node->GetKey()))
+						{
+							const FString Description = Socket->GetDescription(Hierarchy);
+							if(!Description.IsEmpty())
+							{
+								Node->SetLabel(
+									FText::Format(LOCTEXT("NodeLabelFormat", "{0}\n{1}"), 
+									Node->GetLabel(), 
+									FText::FromString(Description))
+								);
+							}
+						}
+					}
+				}
+				break;
+			}
+			default:
+			{
+				break;
+			}
+		}
 
 		RigElementKeyToGuid.Add(Node->GetKey(), Node->GetGuid());
 
+		FRigElementKey ResolvedKey;
+		if(IsConnectorResolved(InKey, &ResolvedKey))
+		{
+			if(const FControlRigSchematicRigElementKeyNode* ResolvedNode = FindElementKeyNode(ResolvedKey))
+			{
+				SetParentNode(Node->GetGuid(), ResolvedNode->GetGuid());
+			}
+		}
+
 		if (bNotify && OnNodeAddedDelegate.IsBound())
 		{
 			OnNodeAddedDelegate.Broadcast(Node);
@@ -100,7 +136,7 @@ bool FControlRigSchematicModel::RemoveNode(const FGuid& InGuid)
 	{
 		KeyToRemove = Node->GetKey();
 	}
-	if(FSchematicGraphModel::RemoveNode(InGuid))
+	if(Super::RemoveNode(InGuid))
 	{
 		RigElementKeyToGuid.Remove(KeyToRemove);
 		return true;
@@ -171,7 +207,12 @@ void FControlRigSchematicModel::UpdateElementKeyLinks()
 				const FControlRigSchematicRigElementKeyNode* ChildNode = VisitElement(Child, Node, OutExpectedLinks);
 				if(Node && ChildNode && Node != ChildNode)
 				{
-					OutExpectedLinks.Emplace(Node->GetKey(), ChildNode->GetKey());
+					const FRigElementKey ParentA = Hierarchy->GetFirstParent(Node->GetKey());
+					const FRigElementKey ParentB = Hierarchy->GetFirstParent(ChildNode->GetKey());
+					if(!ParentA || !ParentB.IsValid() || ParentA != ParentB)
+					{
+						OutExpectedLinks.Emplace(Node->GetKey(), ChildNode->GetKey());
+					}
 				}
 			}
 
@@ -425,6 +466,43 @@ void FControlRigSchematicModel::HandleModularRigModified(EModularRigNotification
 	{
 		case EModularRigNotification::ConnectionChanged:
 		{
+			if (ControlRigBeingDebuggedPtr.IsValid())
+			{
+				if (const URigHierarchy* Hierarchy = ControlRigBeingDebuggedPtr->GetHierarchy())
+				{
+					const TArray<FRigElementKey> Connectors = Hierarchy->GetConnectorKeys();
+					for(const FRigElementKey& ConnectorKey : Connectors)
+					{
+						if(const FSchematicGraphNode* ConnectorNode = FindElementKeyNode(ConnectorKey))
+						{
+							const FString NameString = ConnectorKey.Name.ToString();
+							if(NameString.StartsWith(InModule->GetNamespace(), ESearchCase::CaseSensitive))
+							{
+								const FString LocalName = NameString.Mid(InModule->GetNamespace().Len());
+								if(!LocalName.Contains(UModularRig::NamespaceSeparator))
+								{
+									const FRigElementKey LocalKey(*LocalName, ConnectorKey.Type);
+									if(const FRigElementKey* ResolvedKey = InModule->Connections.Find(LocalKey))
+									{
+										if(const FSchematicGraphNode* ResolvedNode = FindElementKeyNode(*ResolvedKey))
+										{
+											(void)SetParentNode(ConnectorNode, ResolvedNode);
+											const_cast<FSchematicGraphNode*>(ConnectorNode)->OnMouseLeave();
+											return;
+										}
+									}
+									if(const FSchematicGraphNode* ParentNode = ConnectorNode->GetParentNode())
+									{
+										(void)RemoveFromParentNode(ConnectorNode);
+										const_cast<FSchematicGraphNode*>(ConnectorNode)->OnMouseLeave();
+									}
+									return;
+								}
+							}
+						}
+					}
+				}
+			}
 			break;
 		}
 		default:
@@ -438,35 +516,38 @@ FVector2d FControlRigSchematicModel::GetPositionForNode(const FSchematicGraphNod
 {
 	if(const FControlRigSchematicRigElementKeyNode* Node = Cast<FControlRigSchematicRigElementKeyNode>(InNode))
 	{
-		FRigElementKey Key = Node->GetKey();
-		if(Key.Type == ERigElementType::Connector)
+		if(!InNode->HasParentNode())
 		{
-			FRigElementKey ResolvedKey;
-			if(IsConnectorResolved(Key, &ResolvedKey))
+			FRigElementKey Key = Node->GetKey();
+			if(Key.Type == ERigElementType::Connector)
 			{
-				Key = ResolvedKey;
+				FRigElementKey ResolvedKey;
+				if(IsConnectorResolved(Key, &ResolvedKey))
+				{
+					Key = ResolvedKey;
+				}
 			}
-		}
 
-		if(Key.IsValid())
-		{
-			if(const TSharedPtr<FControlRigEditor> Editor = ControlRigEditor.Pin())
+			if(Key.IsValid())
 			{
-				if(Editor.IsValid())
+				if(const TSharedPtr<FControlRigEditor> Editor = ControlRigEditor.Pin())
 				{
-					if(const UControlRig* ControlRig = Editor->GetControlRig())
+					if(Editor.IsValid())
 					{
-						if(const URigHierarchy* Hierarchy = ControlRig->GetHierarchy())
+						if(const UControlRig* ControlRig = Editor->GetControlRig())
 						{
-							const FTransform Transform = Hierarchy->GetGlobalTransform(Key);
-							return Editor->ComputePersonaProjectedScreenPos(Transform.GetLocation());
+							if(const URigHierarchy* Hierarchy = ControlRig->GetHierarchy())
+							{
+								const FTransform Transform = Hierarchy->GetGlobalTransform(Key);
+								return Editor->ComputePersonaProjectedScreenPos(Transform.GetLocation());
+							}
 						}
 					}
 				}
 			}
 		}
 	}
-	return FSchematicGraphModel::GetPositionForNode(InNode);
+	return Super::GetPositionForNode(InNode);
 }
 
 bool FControlRigSchematicModel::GetPositionAnimationEnabledForNode(const FSchematicGraphNode* InNode) const
@@ -475,7 +556,7 @@ bool FControlRigSchematicModel::GetPositionAnimationEnabledForNode(const FSchema
 	{
 		return false;
 	}
-	return FSchematicGraphModel::GetPositionAnimationEnabledForNode(InNode);
+	return Super::GetPositionAnimationEnabledForNode(InNode);
 }
 
 int32 FControlRigSchematicModel::GetNumLayersForNode(const FSchematicGraphNode* InNode) const
@@ -484,7 +565,7 @@ int32 FControlRigSchematicModel::GetNumLayersForNode(const FSchematicGraphNode*
 	{
 		return 3;
 	}
-	return FSchematicGraphModel::GetNumLayersForNode(InNode);
+	return Super::GetNumLayersForNode(InNode);
 }
 
 const FSlateBrush* FControlRigSchematicModel::GetBrushForNode(const FSchematicGraphNode* InNode, int32 InLayerIndex) const
@@ -493,31 +574,66 @@ const FSlateBrush* FControlRigSchematicModel::GetBrushForNode(const FSchematicGr
 	{
 		if(InLayerIndex == 0)
 		{
-			static const FSlateBrush* BackgroundBrush = FControlRigEditorStyle::Get().GetBrush( "ControlRig.Schematic.Background");
+			static const FSlateBrush* BackgroundBrush = FSchematicGraphStyle::Get().GetBrush( "Schematic.Background");
 			return BackgroundBrush;
 		}
 		if(InLayerIndex == 1)
 		{
-			static const FSlateBrush* OutlineBrush = FControlRigEditorStyle::Get().GetBrush( "ControlRig.Schematic.Outline");
-			return OutlineBrush;
+			static const FSlateBrush* OutlineSingleBrush = FSchematicGraphStyle::Get().GetBrush( "Schematic.Outline.Single");
+			static const FSlateBrush* OutlineDoubleBrush = FSchematicGraphStyle::Get().GetBrush( "Schematic.Outline.Double");
+
+			if(Node->GetKey().Type == ERigElementType::Socket)
+			{
+				if (ControlRigBlueprint.IsValid())
+				{
+					int32 Count = 0;
+					for(const TPair<FRigElementKey, FRigElementKey>& Pair : ControlRigBlueprint->ConnectionMap)
+					{
+						if(Pair.Value == Node->GetKey())
+						{
+							Count++;
+							if(Count == 2)
+							{
+								return OutlineDoubleBrush;
+							}
+						}
+					}
+				}
+			}
+			return OutlineSingleBrush;
 		}
 		
 		switch(Node->GetKey().Type)
 		{
 			case ERigElementType::Socket:
 			{
-				static const FSlateBrush* UnresolvedSocketBrush = FControlRigEditorStyle::Get().GetBrush( "ControlRig.Schematic.SocketUnresolved");
-				static const FSlateBrush* ResolvedSocketBrush = FControlRigEditorStyle::Get().GetBrush( "ControlRig.Schematic.SocketResolved");
+				static const FSlateBrush* UnresolvedSocketBrush = FSchematicGraphStyle::Get().GetBrush( "Schematic.Dot.Small");
+				static const FSlateBrush* ResolvedSocketBrush = FSchematicGraphStyle::Get().GetBrush( "Schematic.Dot.Large");
+				static const FSlateBrush* ResolvedMultipleSocketBrush = FSchematicGraphStyle::Get().GetBrush( "Schematic.Dot.Medium");
 
 				if (ControlRigBlueprint.IsValid())
 				{
+					int32 Count = 0;
 					for(const TPair<FRigElementKey, FRigElementKey>& Pair : ControlRigBlueprint->ConnectionMap)
 					{
 						if(Pair.Value == Node->GetKey())
 						{
-							return ResolvedSocketBrush;
+							Count++;
+							if(Count == 2)
+							{
+								break;
+							}
 						}
 					}
+
+					if(Count == 1)
+					{
+						return ResolvedSocketBrush;
+					}
+					if(Count == 2)
+					{
+						return ResolvedMultipleSocketBrush;
+					}
 				}
 
 				return UnresolvedSocketBrush;
@@ -567,7 +683,7 @@ const FSlateBrush* FControlRigSchematicModel::GetBrushForNode(const FSchematicGr
 			}
 		}
 	}
-	return FSchematicGraphModel::GetBrushForNode(InNode, InLayerIndex);
+	return Super::GetBrushForNode(InNode, InLayerIndex);
 }
 
 FLinearColor FControlRigSchematicModel::GetColorForNode(const FSchematicGraphNode* InNode, int32 InLayerIndex) const
@@ -631,60 +747,7 @@ FLinearColor FControlRigSchematicModel::GetColorForNode(const FSchematicGraphNod
 			}
 		}
 	}
-	return FSchematicGraphModel::GetColorForNode(InNode, InLayerIndex);
-}
-
-FText FControlRigSchematicModel::GetToolTipForNode(const FSchematicGraphNode* InNode) const
-{
-	if(const FControlRigSchematicRigElementKeyNode* Node = Cast<FControlRigSchematicRigElementKeyNode>(InNode))
-	{
-		switch(Node->GetKey().Type)
-		{
-			case ERigElementType::Socket:
-			{
-				if (ControlRigBeingDebuggedPtr.IsValid())
-				{
-					if (const URigHierarchy* Hierarchy = ControlRigBeingDebuggedPtr->GetHierarchy())
-					{
-						if(const FRigSocketElement* Socket = Hierarchy->Find<FRigSocketElement>(Node->GetKey()))
-						{
-							const FString Description = Socket->GetDescription(Hierarchy);
-							if(!Description.IsEmpty())
-							{
-								return FText::FromString(Node->GetKey().ToString()+TEXT("\n")+Description);
-							}
-						}
-					}
-				}
-				break;
-			}
-			default:
-			{
-				break;
-			}
-		}
-		return FText::FromString(Node->GetKey().ToString());
-	}
-	return FSchematicGraphModel::GetToolTipForNode(InNode);
-}
-
-ESchematicGraphVisibility::Type FControlRigSchematicModel::GetVisibilityForNode(const FSchematicGraphNode* InNode) const
-{
-	if(const FControlRigSchematicRigElementKeyNode* Node = Cast<FControlRigSchematicRigElementKeyNode>(InNode))
-	{
-		if(Node->GetKey().Type == ERigElementType::Connector)
-		{
-			FRigElementKey ResolvedSocket;
-			if(IsConnectorResolved(Node->GetKey(), &ResolvedSocket))
-			{
-				if(ContainsElementKeyNode(ResolvedSocket))
-				{
-					return ESchematicGraphVisibility::Hidden;
-				}
-			}
-		}
-	}
-	return FSchematicGraphModel::GetVisibilityForNode(InNode);
+	return Super::GetColorForNode(InNode, InLayerIndex);
 }
 
 ESchematicGraphPlacementConstraint::Type FControlRigSchematicModel::GetPlacementForNode(const FSchematicGraphNode* InNode) const
@@ -699,24 +762,27 @@ ESchematicGraphPlacementConstraint::Type FControlRigSchematicModel::GetPlacement
 			}
 		}
 	}
-	return FSchematicGraphModel::GetPlacementForNode(InNode);
+	return Super::GetPlacementForNode(InNode);
 }
 
 const FSlateBrush* FControlRigSchematicModel::GetBrushForLink(const FSchematicGraphLink* InLink) const
 {
-	if(const FSchematicGraphNode* SourceNode = FindNode(InLink->GetSourceNodeGuid()))
+	if(InLink->IsA<FControlRigSchematicRigElementKeyLink>())
 	{
-		if(const FSchematicGraphNode* TargetNode = FindNode(InLink->GetTargetNodeGuid()))
+		if(const FSchematicGraphNode* SourceNode = FindNode(InLink->GetSourceNodeGuid()))
 		{
-			if(SourceNode->IsA<FControlRigSchematicRigElementKeyNode>() &&
-				TargetNode->IsA<FControlRigSchematicRigElementKeyNode>())
+			if(const FSchematicGraphNode* TargetNode = FindNode(InLink->GetTargetNodeGuid()))
 			{
-				static const FSlateBrush* LinkBrush = FControlRigEditorStyle::Get().GetBrush( "ControlRig.Schematic.Link");
-				return LinkBrush;
+				if(SourceNode->IsA<FControlRigSchematicRigElementKeyNode>() &&
+					TargetNode->IsA<FControlRigSchematicRigElementKeyNode>())
+				{
+					static const FSlateBrush* LinkBrush = FControlRigEditorStyle::Get().GetBrush( "ControlRig.Schematic.Link");
+					return LinkBrush;
+				}
 			}
 		}
 	}
-	return FSchematicGraphModel::GetBrushForLink(InLink);
+	return Super::GetBrushForLink(InLink);
 }
 
 FLinearColor FControlRigSchematicModel::GetColorForLink(const FSchematicGraphLink* InLink) const
@@ -733,29 +799,28 @@ FLinearColor FControlRigSchematicModel::GetColorForLink(const FSchematicGraphLin
 			}
 		}
 	}
-	return FSchematicGraphModel::GetColorForLink(InLink);
+	return Super::GetColorForLink(InLink);
 }
 
-bool FControlRigSchematicModel::GetForwardedNodeForDrag(FGuid& InOutGuid) const
+ESchematicGraphVisibility::Type FControlRigSchematicModel::GetVisibilityForTag(const FSchematicGraphTag* InTag) const
 {
-	if(const FControlRigSchematicRigElementKeyNode* ElementKeyNode = Cast<FControlRigSchematicRigElementKeyNode>(FindNode(InOutGuid)))
+	if(InTag->IsA<FSchematicGraphGroupTag>())
 	{
-		for(const TPair<FRigElementKey, FRigElementKey>& Pair : ControlRigBlueprint->ConnectionMap)
+		if(const FControlRigSchematicRigElementKeyNode* Node = Cast<FControlRigSchematicRigElementKeyNode>(InTag->GetNode()))
 		{
-			if(Pair.Value == ElementKeyNode->GetKey())
+			if(Node->GetKey().Type == ERigElementType::Socket)
 			{
-				if(const FControlRigSchematicRigElementKeyNode* ForwardedNode = FindElementKeyNode(Pair.Key))
+				if(Node->GetNumChildNodes() < 3)
 				{
-					InOutGuid = ForwardedNode->GetGuid();
-					return true;
+					return ESchematicGraphVisibility::Hidden;
 				}
 			}
 		}
 	}
-	return FSchematicGraphModel::GetForwardedNodeForDrag(InOutGuid);
+	return Super::GetVisibilityForTag(InTag);
 }
 
-void FControlRigSchematicModel::HandleSchematicNodeClicked(SSchematicGraphPanel* InPanel, SSchematicGraphNode* InNode)
+void FControlRigSchematicModel::HandleSchematicNodeClicked(SSchematicGraphPanel* InPanel, SSchematicGraphNode* InNode, const FPointerEvent& InMouseEvent)
 {
 	for (const TSharedPtr<FSchematicGraphNode>& Node : Nodes)
 	{
@@ -864,7 +929,7 @@ void FControlRigSchematicModel::HandleSchematicBeginDrag(SSchematicGraphPanel* I
 			ExistingElementKeyNode->SetVisibility(Matches.ContainsByPredicate([ExistingElementKeyNode](const FRigElementResolveResult& Match)
 			{
 				return ExistingElementKeyNode->GetKey() == Match.GetKey();
-			}) ? ESchematicGraphVisibility::Visible : ESchematicGraphVisibility::FadedOut);
+			}) ? ESchematicGraphVisibility::Visible : ESchematicGraphVisibility::Hidden);
 		}
 	};
 }
@@ -881,6 +946,7 @@ void FControlRigSchematicModel::HandleSchematicEndDrag(SSchematicGraphPanel* InP
 	{
 		Node->SetVisibility(ESchematicGraphVisibility::Visible);
 	}
+	UpdateElementKeyLinks();
 }
 
 void FControlRigSchematicModel::HandleSchematicDrop(SSchematicGraphPanel* InPanel, SSchematicGraphNode* InNode, const FDragDropEvent& InDragDropEvent)
@@ -1015,10 +1081,22 @@ void FControlRigSchematicModel::HandleSchematicDrop(SSchematicGraphPanel* InPane
 
 bool FControlRigSchematicModel::IsConnectorResolved(const FRigElementKey& InConnectorKey, FRigElementKey* OutKey) const
 {
-	if(ControlRigBlueprint.IsValid())
+	if(ControlRigBlueprint.IsValid() && InConnectorKey.Type == ERigElementType::Connector)
 	{
 		if(const FRigElementKey* TargetKey = ControlRigBlueprint->ConnectionMap.Find(InConnectorKey))
 		{
+			// make sure the target exists
+			if(ControlRigBeingDebuggedPtr.IsValid())
+			{
+				if(const URigHierarchy* DebuggedHierarchy = ControlRigBeingDebuggedPtr->GetHierarchy())
+				{
+					if(!DebuggedHierarchy->Contains(*TargetKey))
+					{
+						return false;
+					}
+				}
+			}
+			
 			if(OutKey)
 			{
 				*OutKey = *TargetKey;
@@ -1028,3 +1106,5 @@ bool FControlRigSchematicModel::IsConnectorResolved(const FRigElementKey& InConn
 	}
 	return false;
 }
+
+#undef LOCTEXT_NAMESPACE
diff --git a/Engine/Plugins/Animation/ControlRig/Source/ControlRigEditor/Private/Editor/ControlRigSchematicModel.h b/Engine/Plugins/Animation/ControlRig/Source/ControlRigEditor/Private/Editor/ControlRigSchematicModel.h
--- a/Engine/Plugins/Animation/ControlRig/Source/ControlRigEditor/Private/Editor/ControlRigSchematicModel.h
+++ b/Engine/Plugins/Animation/ControlRig/Source/ControlRigEditor/Private/Editor/ControlRigSchematicModel.h
@@ -13,11 +13,11 @@ class URigHierarchy;
 struct FRigBaseElement;
 
 /** Node for the schematic view */
-class FControlRigSchematicRigElementKeyNode : public FSchematicGraphNode
+class FControlRigSchematicRigElementKeyNode : public FSchematicGraphGroupNode
 {
 public:
 
-	SCHEMATICGRAPHNODE_BODY(FControlRigSchematicRigElementKeyNode, FSchematicGraphNode)
+	SCHEMATICGRAPHNODE_BODY(FControlRigSchematicRigElementKeyNode, FSchematicGraphGroupNode)
 
 	virtual ~FControlRigSchematicRigElementKeyNode() override {}
 	
@@ -53,11 +53,13 @@ protected:
 	friend class FControlRigSchematicModel;
 };
 
-
 /** Model for the schematic views */
 class FControlRigSchematicModel : public FSchematicGraphModel
 {
 public:
+
+	SCHEMATICGRAPHNODE_BODY(FControlRigSchematicModel, FSchematicGraphModel)
+
 	virtual ~FControlRigSchematicModel() override;
 	
 	void SetEditor(const TSharedRef<FControlRigEditor>& InEditor);
@@ -80,22 +82,19 @@ public:
 	virtual int32 GetNumLayersForNode(const FSchematicGraphNode* InNode) const override;
 	virtual const FSlateBrush* GetBrushForNode(const FSchematicGraphNode* InNode, int32 InLayerIndex) const override;
 	virtual FLinearColor GetColorForNode(const FSchematicGraphNode* InNode, int32 InLayerIndex) const override;
-	virtual FText GetToolTipForNode(const FSchematicGraphNode* InNode) const override;
-	virtual ESchematicGraphVisibility::Type GetVisibilityForNode(const FSchematicGraphNode* InNode) const override;
 	virtual ESchematicGraphPlacementConstraint::Type GetPlacementForNode(const FSchematicGraphNode* InNode) const override;
 	virtual const FSlateBrush* GetBrushForLink(const FSchematicGraphLink* InLink) const override;
 	virtual FLinearColor GetColorForLink(const FSchematicGraphLink* InLink) const override;
-
-	virtual bool GetForwardedNodeForDrag(FGuid& InOutGuid) const override;
+	virtual ESchematicGraphVisibility::Type GetVisibilityForTag(const FSchematicGraphTag* InTag) const override;
 
 private:
 
-	void HandleSchematicNodeClicked(SSchematicGraphPanel* InPanel, SSchematicGraphNode* InNode);
+	void HandleSchematicNodeClicked(SSchematicGraphPanel* InPanel, SSchematicGraphNode* InNode, const FPointerEvent& InMouseEvent);
 	void HandleSchematicBeginDrag(SSchematicGraphPanel* InPanel, SSchematicGraphNode* InNode, const FDragDropOperation& InDragDropOperation);
 	void HandleSchematicEndDrag(SSchematicGraphPanel* InPanel, SSchematicGraphNode* InNode, const FDragDropOperation& InDragDropOperation);
 	void HandleSchematicDrop(SSchematicGraphPanel* InPanel, SSchematicGraphNode* InNode, const FDragDropEvent& InDragDropEvent);
 	bool IsConnectorResolved(const FRigElementKey& InConnectorKey, FRigElementKey* OutKey = nullptr) const;
-	
+
 	TWeakPtr<FControlRigEditor> ControlRigEditor;
 	TWeakObjectPtr<UControlRigBlueprint> ControlRigBlueprint;
 	TWeakObjectPtr<UControlRig> ControlRigBeingDebuggedPtr;
diff --git a/Engine/Source/Editor/AnimationEditorWidgets/Private/AnimationEditorWidgetsModule.cpp b/Engine/Source/Editor/AnimationEditorWidgets/Private/AnimationEditorWidgetsModule.cpp
--- a/Engine/Source/Editor/AnimationEditorWidgets/Private/AnimationEditorWidgetsModule.cpp
+++ b/Engine/Source/Editor/AnimationEditorWidgets/Private/AnimationEditorWidgetsModule.cpp
@@ -2,6 +2,10 @@
 
 #include "Modules/ModuleInterface.h"
 #include "Modules/ModuleManager.h"
+#include "SchematicGraphPanel/SchematicGraphStyle.h"
+#include "SchematicGraphPanel/SchematicGraphDefines.h"
+
+DEFINE_LOG_CATEGORY(LogSchematicGraph);
 
 class FAnimationEditorWidgetsModule
 	: public IModuleInterface
@@ -10,9 +14,20 @@ public:
 
 	// IModuleInterface interface
 	
-	virtual void StartupModule() override { }
-	virtual void ShutdownModule() override { }
+	virtual void StartupModule() override
+	{
+		(void)FSchematicGraphStyle::Get();
+	}
+	
+	virtual void ShutdownModule() override
+	{
+	}
 };
 
+FSchematicGraphStyle& FSchematicGraphStyle::Get()
+{
+	static FSchematicGraphStyle Inst;
+	return Inst;
+}
 
 IMPLEMENT_MODULE(FAnimationEditorWidgetsModule, AnimationEditorWidgets);
diff --git a/Engine/Source/Editor/AnimationEditorWidgets/Private/SchematicGraphPanel/SSchematicGraphPanel.cpp b/Engine/Source/Editor/AnimationEditorWidgets/Private/SchematicGraphPanel/SSchematicGraphPanel.cpp
--- a/Engine/Source/Editor/AnimationEditorWidgets/Private/SchematicGraphPanel/SSchematicGraphPanel.cpp
+++ b/Engine/Source/Editor/AnimationEditorWidgets/Private/SchematicGraphPanel/SSchematicGraphPanel.cpp
@@ -7,6 +7,9 @@
 #include "DragAndDrop/AssetDragDropOp.h"
 #include "Engine/Font.h"
 #include "Engine/Engine.h"
+#include "Animation/AnimatedAttributeManager.h"
+#include <SchematicGraphPanel/SchematicGraphStyle.h>
+#include "Fonts/FontMeasure.h"
 
 #define LOCTEXT_NAMESPACE "SSchematicGraphPanel"
 
@@ -95,6 +98,11 @@ void SSchematicGraphNode::Construct(const FArguments& InArgs)
 	}
 
 	SetVisibility(TAttribute<EVisibility>::CreateSP(this, &SSchematicGraphNode::GetNodeVisibility));
+
+	if(const FSchematicGraphGroupNode* GroupNode = Cast<FSchematicGraphGroupNode>(NodeData))
+	{
+		ExpansionCircleFactor = FFloatAttribute::Create(GroupNode->GetAnimationSettings(), 0.f);
+	}
 }
 
 FVector2D SSchematicGraphNode::ComputeDesiredSize(float LayoutScaleMultiplier) const
@@ -108,6 +116,47 @@ int32 SSchematicGraphNode::OnPaint(const FPaintArgs& Args, const FGeometry& Allo
 
 	const FVector2d CurSize = Size->Get() * Scale->Get();
 	const FVector2d SizeOffset = (CurSize-OriginalSize)*-0.5;
+	static const UFont* Font = GEngine->GetSmallFont();
+	static const FSlateFontInfo SmallFontStyle = Font->GetLegacySlateFontInfo();
+	const bool bIsFadedOut = IsFadedOut();
+	const float FadedOutFactor = bIsFadedOut ? 0.5f : 1.f;
+
+	if(const FSchematicGraphGroupNode* GroupNode = Cast<FSchematicGraphGroupNode>(NodeData))
+	{
+		if(GroupNode->GetExpansionState() > SMALL_NUMBER)
+		{
+			ExpansionCircleFactor->Set(GroupNode == GroupNode->GetGraph()->GetLastExpandedNode() ? 1.f : 0.f);
+		}
+		else
+		{
+			ExpansionCircleFactor->Set(0.f);
+		}
+
+		if(ExpansionCircleFactor.IsValid() && ExpansionCircleFactor->Get() > SMALL_NUMBER)
+		{
+			check(ExpansionCircleFactor.IsValid());
+			
+			ExpansionCircleFactor->Set(1.f);
+			
+			static const FSlateBrush* GroupBrush = FSchematicGraphStyle::Get().GetBrush( "Schematic.Group");
+			const float Radius = GroupNode->GetExpansionState() * (GroupNode->GetExpansionRadius() + FMath::Max(OriginalSize.X, OriginalSize.Y) * 0.525f);
+			const FLinearColor Color = GroupNode->GetExpansionColor() * ExpansionCircleFactor->Get();
+
+			const FVector2d CircleSize = FVector2d::One() * Radius * 2.f;
+			const FVector2d CircleOffset = (CircleSize-OriginalSize)*-0.5;
+
+			NewLayerId++;
+
+			FSlateDrawElement::MakeBox(
+				OutDrawElements,
+				NewLayerId,
+				AllottedGeometry.ToPaintGeometry(CircleSize, FSlateLayoutTransform(CircleOffset)),
+				GroupBrush,
+				ESlateDrawEffect::None,
+				Color * FadedOutFactor
+			);
+		}
+	}
 
 	if(BrushGetter)
 	{
@@ -126,7 +175,7 @@ int32 SSchematicGraphNode::OnPaint(const FPaintArgs& Args, const FGeometry& Allo
 		}
 	}
 
-	if(NodeData && SchematicGraphPanel && SchematicGraphPanel->GraphData)
+	if(NodeData && SchematicGraphPanel && SchematicGraphPanel->GetSchematicGraph())
 	{
 		const float NodeRadius = FMath::Min(CurSize.X, CurSize.Y) * 0.5;
 		
@@ -134,7 +183,7 @@ int32 SSchematicGraphNode::OnPaint(const FPaintArgs& Args, const FGeometry& Allo
 		for(const TSharedPtr<FSchematicGraphTag>& TagPtr : Tags)
 		{
 			const FSchematicGraphTag* CurrentTag = TagPtr.Get();
-			const ESchematicGraphVisibility::Type TagVisibility = SchematicGraphPanel->GraphData->GetVisibilityForTag(CurrentTag);
+			const ESchematicGraphVisibility::Type TagVisibility = SchematicGraphPanel->GetSchematicGraph()->GetVisibilityForTag(CurrentTag);
 			if(TagVisibility == ESchematicGraphVisibility::Hidden)
 			{
 				continue;
@@ -142,18 +191,13 @@ int32 SSchematicGraphNode::OnPaint(const FPaintArgs& Args, const FGeometry& Allo
 
 			FVector2d TagSize = FVector2d::One() * NodeRadius * 1.0f;
 			FVector2d LabelSize = FVector2d::ZeroVector;
-			const FText& Label = SchematicGraphPanel->GraphData->GetLabelForTag(CurrentTag);
+			const FText& Label = SchematicGraphPanel->GetSchematicGraph()->GetLabelForTag(CurrentTag);
 			const FString LabelString = Label.ToString();
 
-			static const UFont* Font = GEngine->GetSmallFont();
-			static const FSlateFontInfo SmallFontStyle = FAppStyle::Get().GetFontStyle("SmallFont");
-
 			if(!LabelString.IsEmpty())
 			{
-				int32 Height = 0, Width = 0;
-				Font->GetStringHeightAndWidth(LabelString, Height, Width);
-				LabelSize.X = Width;
-				LabelSize.Y = Height;
+				const TSharedRef< FSlateFontMeasure > FontMeasureService = FSlateApplication::Get().GetRenderer()->GetFontMeasureService();
+				LabelSize = FontMeasureService->Measure(LabelString, SmallFontStyle);
 				TagSize.X = FMath::Max(TagSize.X, LabelSize.X + 4);
 				TagSize.Y = FMath::Max(TagSize.Y, LabelSize.Y + 4);
 			}
@@ -171,11 +215,11 @@ int32 SSchematicGraphNode::OnPaint(const FPaintArgs& Args, const FGeometry& Allo
 					AllottedGeometry.ToPaintGeometry(TagSize, FSlateLayoutTransform(TagOffset)),
 					BackgroundBrush,
 					ESlateDrawEffect::None,
-					SchematicGraphPanel->GraphData->GetBackgroundColorForTag(CurrentTag)
+					SchematicGraphPanel->GetSchematicGraph()->GetBackgroundColorForTag(CurrentTag) * FadedOutFactor
 				);
 			}
 
-			if(const FSlateBrush* ForegroundBrush = SchematicGraphPanel->GraphData->GetForegroundBrushForTag(CurrentTag))
+			if(const FSlateBrush* ForegroundBrush = SchematicGraphPanel->GetSchematicGraph()->GetForegroundBrushForTag(CurrentTag))
 			{
 				NewLayerId++;
 
@@ -185,7 +229,7 @@ int32 SSchematicGraphNode::OnPaint(const FPaintArgs& Args, const FGeometry& Allo
 					AllottedGeometry.ToPaintGeometry(TagSize, FSlateLayoutTransform(TagOffset)),
 					ForegroundBrush,
 					ESlateDrawEffect::None,
-					SchematicGraphPanel->GraphData->GetForegroundColorForTag(CurrentTag)
+					SchematicGraphPanel->GetSchematicGraph()->GetForegroundColorForTag(CurrentTag) * FadedOutFactor
 				);
 			}
 
@@ -193,7 +237,7 @@ int32 SSchematicGraphNode::OnPaint(const FPaintArgs& Args, const FGeometry& Allo
 			{
 				NewLayerId++;
 				
-				const FVector2d LabelOffset = TagCenter - LabelSize * 0.5 + FVector2d(1, 2);
+				const FVector2d LabelOffset = TagCenter - LabelSize * 0.5;
 
 				FSlateDrawElement::MakeText(
 					OutDrawElements,
@@ -202,18 +246,60 @@ int32 SSchematicGraphNode::OnPaint(const FPaintArgs& Args, const FGeometry& Allo
 					LabelString,
 					SmallFontStyle,
 					ESlateDrawEffect::None,
-					SchematicGraphPanel->GraphData->GetLabelColorForTag(CurrentTag)
+					SchematicGraphPanel->GetSchematicGraph()->GetLabelColorForTag(CurrentTag) * FadedOutFactor
 				);
 			}
 		}
 	}
+
+	const FText& NodeLabel = NodeData->GetLabel();
+	bool bDrawLabel = !NodeLabel.IsEmpty();  
+	if(bDrawLabel)
+	{
+		const FVector2D MouseCursorLocation = FSlateApplication::Get().GetCursorPos();
+		if((AllottedGeometry.GetAbsolutePosition() - MouseCursorLocation).Length() > AllottedGeometry.GetAbsoluteSize().GetMax() * 0.5f + 8.f)
+		{
+			bDrawLabel = false;
+		}
+	}
+	if(bDrawLabel)
+	{
+		if(const FSchematicGraphGroupNode* GroupNode = NodeData->GetGroupNode())
+		{
+			if(GroupNode->IsExpanding() || GroupNode->IsCollapsing())
+			{
+				bDrawLabel = false;
+			}
+		}
+	}
+	if(bDrawLabel)
+	{
+		const FString NodeLabelString = NodeLabel.ToString();
+		const TSharedRef< FSlateFontMeasure > FontMeasureService = FSlateApplication::Get().GetRenderer()->GetFontMeasureService();
+		FVector2d NodeLabelSize = FontMeasureService->Measure(NodeLabelString, SmallFontStyle);
+
+		const FVector2d NodeBottomCenter = AllottedGeometry.GetLocalSize() * FVector2d(0.5f, 1.f) + FVector2d(0.f, 8.f);
+		const FVector2d NodeLabelOffset = NodeBottomCenter - NodeLabelSize * FVector2d(0.5, 0);
+
+		NewLayerId++;
+
+		FSlateDrawElement::MakeText(
+			OutDrawElements,
+			NewLayerId,
+			AllottedGeometry.ToPaintGeometry(NodeLabelSize, FSlateLayoutTransform(NodeLabelOffset)),
+			NodeLabelString,
+			SmallFontStyle,
+			ESlateDrawEffect::None,
+			FLinearColor::White * FadedOutFactor
+		);
+	}
 	
 	return NewLayerId;
 }
 
 void SSchematicGraphNode::OnMouseEnter(const FGeometry& MyGeometry, const FPointerEvent& MouseEvent)
 {
-	if(!IsInteractive())
+	if(!IsInteractive() || IsFadedOut())
 	{
 		return;
 	}
@@ -221,13 +307,13 @@ void SSchematicGraphNode::OnMouseEnter(const FGeometry& MyGeometry, const FPoint
 
 	if(NodeData)
 	{
-		NodeData->SetScaleOffset(ScaledUp);
+		NodeData->OnMouseEnter();
 	}
 }
 
 void SSchematicGraphNode::OnMouseLeave(const FPointerEvent& MouseEvent)
 {
-	if(!IsInteractive())
+	if(!IsInteractive() || IsFadedOut())
 	{
 		return;
 	}
@@ -235,13 +321,13 @@ void SSchematicGraphNode::OnMouseLeave(const FPointerEvent& MouseEvent)
 
 	if(NodeData)
 	{
-		NodeData->SetScaleOffset(1.f);
+		NodeData->OnMouseLeave();
 	}
 }
 
 FReply SSchematicGraphNode::OnDragOver(const FGeometry& MyGeometry, const FDragDropEvent& DragDropEvent)
 {
-	if(!IsInteractive())
+	if(!IsInteractive() || IsFadedOut())
 	{
 		return FReply::Unhandled();
 	}
@@ -249,7 +335,7 @@ FReply SSchematicGraphNode::OnDragOver(const FGeometry& MyGeometry, const FDragD
 
 	if(NodeData)
 	{
-		NodeData->SetScaleOffset(ScaledDown);
+		NodeData->OnDragOver();
 	}
 	return FReply::Handled();
 	
@@ -257,7 +343,7 @@ FReply SSchematicGraphNode::OnDragOver(const FGeometry& MyGeometry, const FDragD
 
 void SSchematicGraphNode::OnDragLeave(const FDragDropEvent& DragDropEvent)
 {
-	if(!IsInteractive())
+	if(!IsInteractive() || IsFadedOut())
 	{
 		return;
 	}
@@ -265,13 +351,13 @@ void SSchematicGraphNode::OnDragLeave(const FDragDropEvent& DragDropEvent)
 
 	if(NodeData)
 	{
-		NodeData->SetScaleOffset(1.f);
+		NodeData->OnDragLeave();
 	}
 }
 
 FReply SSchematicGraphNode::OnDrop(const FGeometry& MyGeometry, const FDragDropEvent& DragDropEvent)
 {
-	if(!IsInteractive())
+	if(!IsInteractive() || IsFadedOut())
 	{
 		return FReply::Unhandled();
 	}
@@ -313,9 +399,9 @@ FReply SSchematicGraphNode::OnDragDetected(const FGeometry& MyGeometry, const FP
 	TArray<FGuid> DraggedElements = {Guid};
 	if (MouseEvent.IsMouseButtonDown(EKeys::LeftMouseButton) && DraggedElements.Num() > 0)
 	{
-		if(SchematicGraphPanel && SchematicGraphPanel->GraphData)
+		if(SchematicGraphPanel && SchematicGraphPanel->GetSchematicGraph())
 		{
-			if(SchematicGraphPanel->GraphData->IsDragSupportedForNode(GetGuid()))
+			if(SchematicGraphPanel->GetSchematicGraph()->IsDragSupportedForNode(GetGuid()))
 			{
 				bIsBeingDragged = true;
 
@@ -340,17 +426,13 @@ EVisibility SSchematicGraphNode::GetNodeVisibility() const
 		return EVisibility::HitTestInvisible;
 	}
 
-	if(SchematicGraphPanel && SchematicGraphPanel->GraphData)
+	if(SchematicGraphPanel && SchematicGraphPanel->GetSchematicGraph())
 	{
-		ESchematicGraphVisibility::Type Vis = SchematicGraphPanel->GraphData->GetVisibilityForNode(GetGuid());
+		ESchematicGraphVisibility::Type Vis = SchematicGraphPanel->GetSchematicGraph()->GetVisibilityForNode(GetGuid());
 		if(Vis == ESchematicGraphVisibility::Hidden)
 		{
 			return EVisibility::Hidden;
 		}
-		if(Vis == ESchematicGraphVisibility::FadedOut)
-		{
-			return EVisibility::HitTestInvisible;
-		}
 	}
 
 	return EVisibility::Visible;
@@ -361,7 +443,23 @@ FReply SSchematicGraphNode::OnMouseButtonDown(const FGeometry& MyGeometry, const
 	SNode::OnMouseButtonDown(MyGeometry, MouseEvent);
 	if (MouseEvent.GetPressedButtons().Contains(EKeys::LeftMouseButton))
 	{
-		OnClickedDelegate.ExecuteIfBound(this);
+		OnClickedDelegate.ExecuteIfBound(this, MouseEvent);
+	}
+
+	const bool bFadedOut = IsFadedOut();
+	
+	if(NodeData)
+	{
+		FReply NodeReply = NodeData->OnClicked(MouseEvent);
+		if(NodeReply.IsEventHandled())
+		{
+			// only allow drag on non-faded nodes
+			if(!bFadedOut)
+			{
+				return NodeReply.DetectDrag(SharedThis(this), EKeys::LeftMouseButton); 
+			}
+			return NodeReply;
+		}
 	}
 	return FReply::Handled().DetectDrag(SharedThis(this), EKeys::LeftMouseButton);
 }
@@ -385,11 +483,23 @@ bool SSchematicGraphNode::IsInteractive() const
 {
 	if(NodeData)
 	{
-		return NodeData->GetVisibility() == ESchematicGraphVisibility::Visible;
+		return NodeData->IsInteractive();
 	}
 	return true;
 }
 
+const bool SSchematicGraphNode::IsFadedOut() const
+{
+	if(NodeData)
+	{
+		if(const FSchematicGraphModel* Graph = NodeData->GetGraph())
+		{
+			return Graph->GetVisibilityForNode(NodeData) == ESchematicGraphVisibility::FadedOut;
+		}
+	}
+	return false;
+}
+
 void SSchematicGraphPanel::SetSchematicGraph(FSchematicGraphModel* InGraphData)
 {
 	if (GraphData)
@@ -476,7 +586,7 @@ void SSchematicGraphPanel::AddNode(const FSchematicGraphNode* InNodeToAdd)
 	const FGuid Guid = InNodeToAdd->GetGuid();
 	const auto Position = FVector2dAttribute::CreateWithGetter(Vector2DInterpolationSettings, FVector2dAttribute::FGetter::CreateSP(this, &SSchematicGraphPanel::GetPositionForNode, Guid));
 	const auto Size = FVector2dAttribute::CreateWithGetter(Vector2DInterpolationSettings, FVector2dAttribute::FGetter::CreateRaw(GraphData, &FSchematicGraphModel::GetSizeForNode, InNodeToAdd));
-	const auto Scale = FFloatAttribute::CreateWithGetter(FloatInterpolationSettings, FFloatAttribute::FGetter::CreateSP(this, &SSchematicGraphPanel::GetScaleForNode, Guid, true), 0.f);
+	const auto Scale = FFloatAttribute::CreateWithGetter(FloatInterpolationSettings, FFloatAttribute::FGetter::CreateSP(this, &SSchematicGraphPanel::GetScaleForNode, Guid), 0.f);
 
 	const int32 NumLayers = GraphData ? GraphData->GetNumLayersForNode(InNodeToAdd) : InNodeToAdd->GetNumLayers();
 	TArray<TSharedPtr<FLinearColorAttribute>> Colors;
@@ -515,8 +625,8 @@ void SSchematicGraphPanel::RemoveNode(const FSchematicGraphNode* InNodeToRemove)
 	
 	for (int32 Iter = 0; Iter != Children.Num(); ++Iter)
 	{
-		TSharedRef<SSchematicGraphNode> Child = GetChild(Iter);
-		if (Child->GetGuid() == GuidToRemove)
+		TSharedRef<SSchematicGraphNode> Widget = GetChild(Iter);
+		if (Widget->GetGuid() == GuidToRemove)
 		{
 			Children.RemoveAt(Iter);
 			break;
@@ -524,8 +634,8 @@ void SSchematicGraphPanel::RemoveNode(const FSchematicGraphNode* InNodeToRemove)
 	}
 	for (int32 Iter = 0; Iter != VisibleChildren.Num(); ++Iter)
 	{
-		TSharedRef<SSchematicGraphNode> Child = StaticCastSharedRef<SSchematicGraphNode>(VisibleChildren[Iter]);
-		if (Child->GetGuid() == GuidToRemove)
+		TSharedRef<SSchematicGraphNode> Widget = StaticCastSharedRef<SSchematicGraphNode>(VisibleChildren[Iter]);
+		if (Widget->GetGuid() == GuidToRemove)
 		{
 			VisibleChildren.RemoveAt(Iter);
 			break;
@@ -542,7 +652,13 @@ const SSchematicGraphNode* SSchematicGraphPanel::FindNode(const FGuid& InGuid) c
 	return nullptr;
 }
 
- void SSchematicGraphPanel::AddLink(const FSchematicGraphLink* InLinkToAdd)
+SSchematicGraphNode* SSchematicGraphPanel::FindNode(const FGuid& InGuid)
+{
+	const SSchematicGraphPanel* ConstThis = this;
+	return const_cast<SSchematicGraphNode*>(ConstThis->FindNode(InGuid));
+}
+
+void SSchematicGraphPanel::AddLink(const FSchematicGraphLink* InLinkToAdd)
  {
 	if(GraphData == nullptr)
 	{
@@ -583,7 +699,13 @@ const SSchematicGraphNode* SSchematicGraphPanel::FindNode(const FGuid& InGuid) c
 	return nullptr;
  }
 
- void SSchematicGraphPanel::OnArrangeChildren(const FGeometry& AllottedGeometry, FArrangedChildren& ArrangedChildren) const
+SSchematicGraphPanel::FSchematicLinkWidgetInfo* SSchematicGraphPanel::FindLink(const FGuid& InGuid)
+{
+	const SSchematicGraphPanel* ConstThis = this;
+	return const_cast<FSchematicLinkWidgetInfo*>(ConstThis->FindLink(InGuid));
+}
+
+void SSchematicGraphPanel::OnArrangeChildren(const FGeometry& AllottedGeometry, FArrangedChildren& ArrangedChildren) const
 {
 	SNodePanel::OnArrangeChildren(AllottedGeometry, ArrangedChildren);
 }
@@ -641,13 +763,50 @@ int32 SSchematicGraphPanel::OnPaint(const FPaintArgs& Args, const FGeometry& All
 
 		const FSchematicGraphNode* NodeData = ChildNode->GetNodeData();
 
-		const bool bIsVisible =
-			!CurWidget.Geometry.GetLocalSize().IsNearlyZero() &&
-			FSlateRect::DoRectanglesIntersect( CurWidget.Geometry.GetLayoutBoundingRect(), MyCullingRect ) &&
-			(GraphData->GetVisibilityForNode(NodeData) != ESchematicGraphVisibility::Hidden);
+		ESchematicGraphVisibility::Type NodeVisibility = ESchematicGraphVisibility::Visible;
+		if(CurWidget.Geometry.GetLocalSize().IsNearlyZero() ||
+			!FSlateRect::DoRectanglesIntersect( CurWidget.Geometry.GetLayoutBoundingRect(), MyCullingRect ))
+		{
+			NodeVisibility = ESchematicGraphVisibility::Hidden;
+		}
+		else
+		{
+			NodeVisibility = GraphData->GetVisibilityForNode(NodeData);
+		}
+
+		NodeVisibilityByIndex.Add(NodeVisibility);
+		NodeVisibilityByGuid.Add(NodeData->GetGuid(), NodeVisibility);
+	}
+
+	// update the node visibility and centers based on the group relationships
+	for (const TSharedPtr<FSchematicGraphNode>& Node : GraphData->GetNodes())
+	{
+		if(Node->IsRootNode())
+		{
+			continue;
+			
+		}
+
+		// if the node is already visible
+		if(const ESchematicGraphVisibility::Type* ChildVisibility = NodeVisibilityByGuid.Find(Node->GetGuid()))
+		{
+			if(*ChildVisibility != ESchematicGraphVisibility::Hidden)
+			{
+				continue;
+			}
+		}
 
-		NodeVisibilityByIndex.Add(bIsVisible);
-		NodeVisibilityByGuid.Add(NodeData->GetGuid(), bIsVisible);
+		const FGuid RootGuid = Node->GetRootNodeGuid();
+		if(const ESchematicGraphVisibility::Type* RootNodeVisibility = NodeVisibilityByGuid.Find(RootGuid))
+		{
+			// update the child node's visibility + center
+			if(*RootNodeVisibility != ESchematicGraphVisibility::Hidden)
+			{
+				NodeVisibilityByGuid.FindOrAdd(Node->GetGuid()) = *RootNodeVisibility;
+				const FVector2d& RootNodeCenter = NodeCenterByGuid.FindChecked(RootGuid);
+				NodeCenterByGuid.FindOrAdd(Node->GetGuid()) = RootNodeCenter;
+			}
+		}
 	}
 
 	// draw all of the links
@@ -665,16 +824,18 @@ int32 SSchematicGraphPanel::OnPaint(const FPaintArgs& Args, const FGeometry& All
 		}
 
 		// the node may have been culled
-		const bool* bIsSourceNodeVisible = NodeVisibilityByGuid.Find(Link->GetSourceNodeGuid()); 
-		const bool* bIsTargetNodeVisible = NodeVisibilityByGuid.Find(Link->GetTargetNodeGuid());
-		if((bIsSourceNodeVisible == nullptr) || (bIsTargetNodeVisible == nullptr))
+		const ESchematicGraphVisibility::Type* IsSourceNodeVisible = NodeVisibilityByGuid.Find(Link->GetSourceNodeGuid()); 
+		const ESchematicGraphVisibility::Type* IsTargetNodeVisible = NodeVisibilityByGuid.Find(Link->GetTargetNodeGuid());
+		if((IsSourceNodeVisible == nullptr) || (IsTargetNodeVisible == nullptr))
 		{
 			continue;
 		}
-		if(!(*bIsSourceNodeVisible) || !(*bIsTargetNodeVisible))
+		if((*IsSourceNodeVisible == ESchematicGraphVisibility::Hidden) || (*IsTargetNodeVisible == ESchematicGraphVisibility::Hidden))
 		{
 			continue;
 		}
+
+		const bool bFadedOut = (*IsSourceNodeVisible == ESchematicGraphVisibility::FadedOut) || (*IsTargetNodeVisible == ESchematicGraphVisibility::FadedOut);
 		
 		const SSchematicGraphNode* SourceNode = FindNode(Link->GetSourceNodeGuid()); 
 		const SSchematicGraphNode* TargetNode = FindNode(Link->GetTargetNodeGuid());
@@ -683,12 +844,17 @@ int32 SSchematicGraphPanel::OnPaint(const FPaintArgs& Args, const FGeometry& All
 			continue;
 		}
 
-		const FLinearColor Color = Pair.Value->Color->Get();
+		const FLinearColor Color = Pair.Value->Color->Get() * (bFadedOut ? 0.5f : 1.f);
 		const float Thickness = Pair.Value->Thickness->Get();
 		const FSlateBrush* Brush = GraphData->GetBrushForLink(Pair.Key);
 		const FVector2d& SourcePosition = NodeCenterByGuid.FindChecked(SourceNode->GetGuid()) + GraphData->GetSourceNodeOffsetForLink(Link);
 		const FVector2d& TargetPosition = NodeCenterByGuid.FindChecked(TargetNode->GetGuid()) + GraphData->GetTargetNodeOffsetForLink(Link);
 
+		if(SourcePosition.IsNearlyZero() || TargetPosition.IsNearlyZero())
+		{
+			continue;
+		}
+
 		const FVector2d Diff = TargetPosition - SourcePosition;
 		const float DiffLength = Diff.Size();
 		if(DiffLength < SMALL_NUMBER)
@@ -762,7 +928,7 @@ int32 SSchematicGraphPanel::OnPaint(const FPaintArgs& Args, const FGeometry& All
 	{
 		for (int32 ChildIndex = 0; ChildIndex < ArrangedChildren.Num(); ++ChildIndex)
 		{
-			if(!NodeVisibilityByIndex[ChildIndex])
+			if(NodeVisibilityByIndex[ChildIndex] == ESchematicGraphVisibility::Hidden)
 			{
 				continue;
 			}
@@ -808,6 +974,11 @@ TStatId SSchematicGraphPanel::GetStatId() const
 
 void SSchematicGraphPanel::Tick(float DeltaTime)
 {
+	if(GraphData)
+	{
+		GraphData->Tick(DeltaTime);
+	}
+	
 	bIsDragDropping = false;
 	FSlateApplication& Application = FSlateApplication::Get();
 	if (Application.IsDragDropping())
@@ -819,6 +990,17 @@ void SSchematicGraphPanel::Tick(float DeltaTime)
 		}
 	}
 
+	if(bIsDragDropping)
+	{
+		const FVector2D MouseCursorLocation = FSlateApplication::Get().GetCursorPos();
+		if(GetPaintSpaceGeometry().GetRenderBoundingRect().ContainsPoint(MouseCursorLocation))
+		{
+			// when dragging and dropping we want the UI to update
+			// for that we are forcing an update on the tick of the animated attributes
+			FAnimatedAttributeManager::Get().Tick(DeltaTime);
+		}
+	}
+
 	NodesTopLeft.Reset();
 	NodesTopRight.Reset();
 	NodesBottomLeft.Reset();
@@ -828,13 +1010,13 @@ void SSchematicGraphPanel::Tick(float DeltaTime)
 
 	for (int32 i=0; i<Children.Num(); ++i)
 	{
-		TSharedRef<SSchematicGraphNode> Node = GetChild(i);
+		TSharedRef<SSchematicGraphNode> Widget = GetChild(i);
 
 		// update the animation state of the node
-		Node->EnablePositionAnimation(GraphData->GetPositionAnimationEnabledForNode(Node->GetGuid()));
+		Widget->EnablePositionAnimation(GraphData->GetPositionAnimationEnabledForNode(Widget->GetGuid()));
 
 		// collect the nodes displayed in the corners
-		switch(GraphData->GetPlacementForNode(Node->GetNodeData()))
+		switch(GraphData->GetPlacementForNode(Widget->GetNodeData()))
 		{
 			case ESchematicGraphPlacementConstraint::Free:
 			{
@@ -842,27 +1024,27 @@ void SSchematicGraphPanel::Tick(float DeltaTime)
 			}
 			case ESchematicGraphPlacementConstraint::TopLeft:
 			{
-				NodesTopLeft.Add(Node->GetGuid());
+				NodesTopLeft.Add(Widget->GetGuid());
 				break;
 			}
 			case ESchematicGraphPlacementConstraint::TopRight:
 			{
-				NodesTopRight.Add(Node->GetGuid());
+				NodesTopRight.Add(Widget->GetGuid());
 				break;
 			}
 			case ESchematicGraphPlacementConstraint::BottomLeft:
 			{
-				NodesBottomLeft.Add(Node->GetGuid());
+				NodesBottomLeft.Add(Widget->GetGuid());
 				break;
 			}
 			case ESchematicGraphPlacementConstraint::BottomRight:
 			{
-				NodesBottomRight.Add(Node->GetGuid());
+				NodesBottomRight.Add(Widget->GetGuid());
 				break;
 			}
 		}
 
-		if (Node->IsBeingDragged())
+		if (Widget->IsBeingDragged())
 		{
 			if (bIsDragDropping)
 			{
@@ -870,18 +1052,21 @@ void SSchematicGraphPanel::Tick(float DeltaTime)
 				const FGeometry& Geometry = GetTickSpaceGeometry();
 
 				const FVector2d LocalMousePosition =  (AbsoluteMousePosition - Geometry.GetAbsolutePosition()) / Geometry.GetAccumulatedLayoutTransform().GetScale();
-				const FVector2d HalfOriginalSize = Node->GetOriginalSize() * 0.5;
-				Node->PositionDuringDrag = LocalMousePosition + HalfOriginalSize;
+				const FVector2d HalfOriginalSize = Widget->GetOriginalSize() * 0.5;
+				Widget->PositionDuringDrag = LocalMousePosition + HalfOriginalSize;
 			}
 			else if (DeltaTime > 0.f)
 			{
-				Node->PositionDuringDrag.Reset();
-				Node->OffsetDuringDrag.Reset();
-				Node->bIsBeingDragged = false;
+				Widget->PositionDuringDrag.Reset();
+				Widget->OffsetDuringDrag.Reset();
+				Widget->bIsBeingDragged = false;
 			}
 		}
 	}
 
+	UpdatePerNodeCaches(true);
+	UpdateAutoGroupingForNodes();
+	UpdatePerNodeCaches(false);
 	UpdateAutoScalingForNodes();
 }
 
@@ -891,9 +1076,9 @@ void SSchematicGraphPanel::Tick(const FGeometry& AllottedGeometry, const double
 	Tick(InDeltaTime);
 }
 
-void SSchematicGraphPanel::OnNodeClicked(SSchematicGraphNode* Node)
+void SSchematicGraphPanel::OnNodeClicked(SSchematicGraphNode* Node, const FPointerEvent& MouseEvent)
 {
-	OnNodeClickedDelegate.ExecuteIfBound(this, Node);
+	OnNodeClickedDelegate.ExecuteIfBound(this, Node, MouseEvent);
 }
 
 void SSchematicGraphPanel::OnBeginDragEvent(SSchematicGraphNode* Node, const FDragDropOperation& InDragDropOp)
@@ -1002,7 +1187,7 @@ FLinearColor SSchematicGraphPanel::GetColorForNode(FGuid InNodeGuid, int32 InLay
 	return FLinearColor::White;
 }
 
-float SSchematicGraphPanel::GetScaleForNode(FGuid InNodeGuid, bool bIncludeScaleOffset) const
+float SSchematicGraphPanel::GetScaleForNode(FGuid InNodeGuid) const
 {
 	// check if the node may be auto scaled
 	if(const TSharedPtr<SSchematicGraphNode>* NodePtr = NodeByGuid.Find(InNodeGuid))
@@ -1020,11 +1205,29 @@ float SSchematicGraphPanel::GetScaleForNode(FGuid InNodeGuid, bool bIncludeScale
 	
 	if(GraphData)
 	{
-		return GraphData->GetScaleForNode(InNodeGuid, bIncludeScaleOffset);
+		return GraphData->GetScaleForNode(InNodeGuid);
 	}
 	return 1.f;
 }
 
+bool SSchematicGraphPanel::IsAutoGroupingEnabled() const
+{
+	if(GraphData)
+	{
+		return GraphData->IsAutoGroupingEnabled();
+	}
+	return false;
+}
+
+float SSchematicGraphPanel::GetAutoGroupingDistance() const
+{
+	if(GraphData)
+	{
+		return GraphData->GetAutoGroupingDistance();
+	}
+	return 0.f;
+}
+
 bool SSchematicGraphPanel::IsAutoScaleEnabledForNode(FGuid InNodeGuid) const
 {
 	if(GraphData)
@@ -1041,7 +1244,7 @@ bool SSchematicGraphPanel::IsAutoScaleEnabledForNode(FGuid InNodeGuid) const
 		const float MinimumDistance = GraphData->GetMinimumLinkDistanceForNode(InLinkGuid);
 		if(bIncludeScale)
 		{
-			const float Scale = GetScaleForNode(InLinkGuid, true);
+			const float Scale = GetScaleForNode(InLinkGuid);
 			return MinimumDistance * Scale;
 		}
 		return MinimumDistance;
@@ -1049,39 +1252,184 @@ bool SSchematicGraphPanel::IsAutoScaleEnabledForNode(FGuid InNodeGuid) const
 	return 0.f;
  }
 
-void SSchematicGraphPanel::UpdateAutoScalingForNodes()
+void SSchematicGraphPanel::UpdatePerNodeCaches(bool bRemoveNodesFromAutoGroups)
 {
+	PerNodeCaches.Reset();
+
 	if(GraphData == nullptr)
 	{
 		return;
 	}
-	
-	TArray<bool> NodeIsAutoScaling;
-	TArray<FVector2d> NodePositions;
-	TArray<double> NodeRadiuses;
-	NodeIsAutoScaling.Reserve(Children.Num());
-	NodePositions.Reserve(Children.Num());
-	NodeRadiuses.Reserve(Children.Num());
+
+	PerNodeCaches.Reserve(Children.Num());
 
 	for (int32 i=0; i<Children.Num(); ++i)
 	{
-		TSharedRef<SSchematicGraphNode> Node = GetChild(i);
+		TSharedRef<SSchematicGraphNode> Widget = GetChild(i);
 
-		NodeIsAutoScaling.Add(!Node->bIsBeingDragged && Node->EnableAutoScale.Get());
-		if(NodeIsAutoScaling.Last())
+		FPerNodeCache Cache;
+		if(FSchematicGraphNode* Node = Widget->GetNodeData())
 		{
-			NodePositions.Add(GetPositionForNode(Node->GetGuid()));
+			if(bRemoveNodesFromAutoGroups)
+			{
+				if(Cast<FSchematicGraphAutoGroupNode>(Node->GetParentNode()))
+				{
+					GraphData->RemoveFromParentNode(Node, false);
+				}
+			}
+			Cache.bHasParent = Node->HasParentNode();
+			Cache.Visibility = GraphData->GetVisibilityForNode(Node);
+			Cache.bIsAutoScaling = !Cache.bHasParent && !Widget->bIsBeingDragged && Widget->EnableAutoScale.Get();
+			Cache.Position = GetPositionForNode(Node->GetGuid());
 			const FVector2d NodeSize = GraphData->GetSizeForNode(Node->GetGuid());
 			// note: this is not necessarily the best way to determine the radius of a node
-			NodeRadiuses.Add(FMath::Min(NodeSize.X, NodeSize.Y) * 0.5);
+			Cache.Radius = FMath::Min(NodeSize.X, NodeSize.Y) * 0.5;
+		}
+
+		PerNodeCaches.Add(Cache);
+	}
+}
+
+void SSchematicGraphPanel::UpdateAutoGroupingForNodes()
+{
+	if(GraphData == nullptr)
+	{
+		return;
+	}
+
+	if(!IsAutoGroupingEnabled())
+	{
+		return;
+	}
+
+	TMap<uint32, FSchematicGraphGroupNode*> GroupNodeByHash;
+	for(const TPair<uint32, FGuid>& Pair : GroupNodeGuidByHash)
+	{
+		if(FSchematicGraphGroupNode* GroupNode = Cast<FSchematicGraphGroupNode>(GraphData->FindNode(Pair.Value)))
+		{
+			GroupNodeByHash.Add(Pair.Key, GroupNode);
+		}
+	}
+
+	// update the group nodes as needed
+	const float AutoGroupingDistance = GetAutoGroupingDistance();
+	TMap<uint32, TArray<FGuid>> NodeGuidsPerHash;
+	TMap<uint32, FVector2d> NodePositionPerHash;
+	for (int32 i=0; i<Children.Num(); ++i)
+	{
+		if(PerNodeCaches[i].bHasParent)
+		{
+			continue;
+		}
+		if(PerNodeCaches[i].Visibility == ESchematicGraphVisibility::Hidden)
+		{
+			continue;
+		}
+
+		TSharedRef<SSchematicGraphNode> Widget = GetChild(i);
+		const FSchematicGraphNode* Node = Widget->GetNodeData();
+		if(Node->IsA<FSchematicGraphAutoGroupNode>())
+		{
+			continue;
+		}
+
+		const FVector2d FloatingPointPosition = PerNodeCaches[i].Position;
+		const TTuple<int32,int32> IntegerPosition = {
+			FMath::RoundToInt(FloatingPointPosition.X / AutoGroupingDistance),
+			FMath::RoundToInt(FloatingPointPosition.Y / AutoGroupingDistance)
+		};
+		
+		const uint32 PositionHash = HashCombine(IntegerPosition.Get<0>(), IntegerPosition.Get<1>());
+		NodeGuidsPerHash.FindOrAdd(PositionHash).Add(Node->GetGuid());
+
+		if(!NodePositionPerHash.Contains(PositionHash))
+		{
+			NodePositionPerHash.Add(PositionHash, FloatingPointPosition);
+		}
+	}
+
+	// remove all groups which has 0 or 1 element.
+	NodeGuidsPerHash = NodeGuidsPerHash.FilterByPredicate([](const TPair<uint32, TArray<FGuid>>& Pair) -> bool
+	{
+		return Pair.Value.Num() > 1;
+	});
+
+	// now that we have the groupings - let's update the nodes
+	TMap<uint32, uint32> CombinedHashToPositionHash;
+	for(const TPair<uint32, TArray<FGuid>>& Pair : NodeGuidsPerHash)
+	{
+		uint32 CombinedGuidHash = 0;
+		for(const FGuid& ChildNodeGuid : Pair.Value)
+		{
+			CombinedGuidHash = HashCombine(CombinedGuidHash, GetTypeHash(ChildNodeGuid));
+		}
+		CombinedHashToPositionHash.Add(CombinedGuidHash, Pair.Key);
+		
+		FSchematicGraphGroupNode* GroupNode = nullptr;
+		if(FSchematicGraphGroupNode** ExistingGroupNode = GroupNodeByHash.Find(CombinedGuidHash))
+		{
+			GroupNode = *ExistingGroupNode;
 		}
 		else
 		{
-			NodePositions.Add(FVector2d::ZeroVector);
-			NodeRadiuses.Add(0);
+			GroupNode = GraphData->AddAutoGroupNode();
+			GroupNodeByHash.Add(CombinedGuidHash, GroupNode);
+		}
+
+		const TArray<FGuid> PreviousChildNodeGuids = GroupNode->GetChildNodeGuids();
+		const TArray<FGuid>& NextChildNodeGuids = Pair.Value;
+
+		for(const FGuid& NextChildNodeGuid : NextChildNodeGuids)
+		{
+			if(!PreviousChildNodeGuids.Contains(NextChildNodeGuid))
+			{
+				GraphData->SetParentNode(NextChildNodeGuid, GroupNode->GetGuid());
+			}
 		}
 	}
 
+	// update the guid based map based on the updated existing group nodes
+	GroupNodeGuidByHash.Reset();
+	for(const TPair<uint32, FSchematicGraphGroupNode*>& Pair : GroupNodeByHash)
+	{
+		// remove redundant nodes
+		const uint32* PositionHash = CombinedHashToPositionHash.Find(Pair.Key);
+		if(PositionHash == nullptr)
+		{
+			GraphData->RemoveNode(Pair.Value->GetGuid());
+			continue;
+		}
+		const TArray<FGuid>* Guids = NodeGuidsPerHash.Find(*PositionHash); 
+		if(Guids == nullptr)
+		{
+			GraphData->RemoveNode(Pair.Value->GetGuid());
+			continue;
+		}
+
+		FSchematicGraphNode* Node = Pair.Value;
+		GroupNodeGuidByHash.Add(Pair.Key, Node->GetGuid());
+
+		// move the group node to the right location
+		const FVector2d AveragePosition = NodePositionPerHash.FindChecked(*PositionHash); 
+		Pair.Value->SetPosition(AveragePosition);
+
+		// also update the widget since it has an animated position
+		if(const SSchematicGraphNode* Widget = FindNode(Node->GetGuid()))
+		{
+			Widget->Position->SetValueAndStop(AveragePosition);
+			Widget->Scale->SetValueAndStop(1.f);
+		}
+		
+	}
+}
+
+void SSchematicGraphPanel::UpdateAutoScalingForNodes()
+{
+	if(GraphData == nullptr)
+	{
+		return;
+	}
+	
 	// for now brute force find all neighbors
 	// and determine how much of the radius we have
 	// to reduce to avoid overlap.
@@ -1091,23 +1439,23 @@ void SSchematicGraphPanel::UpdateAutoScalingForNodes()
 	
 	for (int32 i=0; i<Children.Num(); ++i)
 	{
-		if(!NodeIsAutoScaling[i])
+		if(!PerNodeCaches[i].bIsAutoScaling)
 		{
 			continue;
 		}
 		
-		const FVector2d& PositionA = NodePositions[i];
-		const double RadiusA = NodeRadiuses[i];
+		const FVector2d& PositionA = PerNodeCaches[i].Position;
+		const double RadiusA = PerNodeCaches[i].Radius;
 
 		for (int32 j=i+1; j<Children.Num(); ++j)
 		{
-			if(!NodeIsAutoScaling[j])
+			if(!PerNodeCaches[j].bIsAutoScaling)
 			{
 				continue;
 			}
 
-			const FVector2d& PositionB = NodePositions[j];
-			const double RadiusB = NodeRadiuses[j];
+			const FVector2d& PositionB = PerNodeCaches[j].Position;
+			const double RadiusB = PerNodeCaches[j].Radius;
 			static constexpr double AutoScalePadding = 4.0;
 			const double MinDistance = RadiusA + RadiusB + AutoScalePadding;
 
@@ -1126,15 +1474,15 @@ void SSchematicGraphPanel::UpdateAutoScalingForNodes()
 	// mark nodes for auto scaling
 	for (int32 i=0; i<Children.Num(); ++i)
 	{
-		TSharedRef<SSchematicGraphNode> Node = GetChild(i);
+		TSharedRef<SSchematicGraphNode> Widget = GetChild(i);
 		if(RadiusReductionPerNode[i] > SMALL_NUMBER)
 		{
-			const float Scale = (NodeRadiuses[i] - RadiusReductionPerNode[i]) / NodeRadiuses[i];
-			Node->AutoScale = Scale;
+			const float Scale = (PerNodeCaches[i].Radius - RadiusReductionPerNode[i]) / PerNodeCaches[i].Radius;
+			Widget->AutoScale = FMath::Max(Scale, 0.4f);
 		}
 		else
 		{
-			Node->AutoScale.Reset();
+			Widget->AutoScale.Reset();
 		}
 	}
 }
diff --git a/Engine/Source/Editor/AnimationEditorWidgets/Private/SchematicGraphPanel/SchematicGraphModel.cpp b/Engine/Source/Editor/AnimationEditorWidgets/Private/SchematicGraphPanel/SchematicGraphModel.cpp
--- a/Engine/Source/Editor/AnimationEditorWidgets/Private/SchematicGraphPanel/SchematicGraphModel.cpp
+++ b/Engine/Source/Editor/AnimationEditorWidgets/Private/SchematicGraphPanel/SchematicGraphModel.cpp
@@ -36,12 +36,25 @@ bool FSchematicGraphModel::RemoveNode(const FGuid& InNodeGuid)
 				(void)RemoveLink(LinkGuid);
 			}
 		}
-		
+
+		if(LastExpandedNode == InNodeGuid)
+		{
+			LastExpandedNode = FGuid();
+		}
+			
 		if (OnNodeRemovedDelegate.IsBound())
 		{
 			OnNodeRemovedDelegate.Broadcast(Node);
 		}
+
+		const TArray<FGuid> ChildNodeGuids = Node->GetChildNodeGuids(); 
+		for(const FGuid& ChildNodeGuid : ChildNodeGuids)
+		{
+			(void)RemoveFromParentNode(ChildNodeGuid);
+		}
+		
 		const FGuid NodeGuid = Node->GetGuid();
+		(void)RemoveFromParentNode(NodeGuid);
 		NodeByGuid.Remove(Node->GetGuid());
 		Nodes.RemoveAll([NodeGuid](const TSharedPtr<FSchematicGraphNode>& ExistingNode) -> bool
 		{
@@ -52,16 +65,16 @@ bool FSchematicGraphModel::RemoveNode(const FGuid& InNodeGuid)
 	return false;
 }
 
-bool FSchematicGraphModel::SetParentNode(const FGuid& InChildNodeGuid, const FGuid& InParentNodeGuid)
+bool FSchematicGraphModel::SetParentNode(const FGuid& InChildNodeGuid, const FGuid& InParentNodeGuid, bool bUpdateGroupNode)
 {
 	if(const FSchematicGraphNode* Node = FindNode(InChildNodeGuid))
 	{
-		return SetParentNode(Node, FindNode(InParentNodeGuid));
+		return SetParentNode(Node, FindNode(InParentNodeGuid), bUpdateGroupNode);
 	}
 	return false;
 }
 
-bool FSchematicGraphModel::SetParentNode(const FSchematicGraphNode* InChildNode, const FSchematicGraphNode* InParentNode)
+bool FSchematicGraphModel::SetParentNode(const FSchematicGraphNode* InChildNode, const FSchematicGraphNode* InParentNode, bool bUpdateGroupNode)
 {
 	check(InChildNode);
 
@@ -72,7 +85,7 @@ bool FSchematicGraphModel::SetParentNode(const FSchematicGraphNode* InChildNode,
 			return false;
 		}
 		
-		(void)RemoveFromParentNode(InChildNode);
+		(void)RemoveFromParentNode(InChildNode, bUpdateGroupNode);
 		const_cast<FSchematicGraphNode*>(InParentNode)->ChildNodeGuids.AddUnique(InChildNode->GetGuid());
 		const_cast<FSchematicGraphNode*>(InChildNode)->ParentNodeGuid = InParentNode->GetGuid();
 		return true;
@@ -80,7 +93,19 @@ bool FSchematicGraphModel::SetParentNode(const FSchematicGraphNode* InChildNode,
 	
 	if(const FSchematicGraphNode* CurrentParentNode = InChildNode->GetParentNode())
 	{
-		(void)const_cast<FSchematicGraphNode*>(CurrentParentNode)->ChildNodeGuids.Remove(InChildNode->GetGuid());
+		FSchematicGraphNode* MutableParentNode = const_cast<FSchematicGraphNode*>(CurrentParentNode);
+		if(bUpdateGroupNode && (MutableParentNode->ChildNodeGuids.Num() == 1))
+		{
+			if(FSchematicGraphGroupNode* GroupParentNode = Cast<FSchematicGraphGroupNode>(MutableParentNode))
+			{
+				GroupParentNode->SetExpanded(false);
+				if(LastExpandedNode == GroupParentNode->GetGuid())
+				{
+					ClearLastExpandedNode();
+				}
+			}
+		}
+		(void)MutableParentNode->ChildNodeGuids.Remove(InChildNode->GetGuid());
 		const_cast<FSchematicGraphNode*>(InChildNode)->ParentNodeGuid = FGuid();
 		return true;
 	}
@@ -88,19 +113,19 @@ bool FSchematicGraphModel::SetParentNode(const FSchematicGraphNode* InChildNode,
 	return false;
 }
 
-bool FSchematicGraphModel::RemoveFromParentNode(const FGuid& InChildNodeGuid)
+bool FSchematicGraphModel::RemoveFromParentNode(const FGuid& InChildNodeGuid, bool bUpdateGroupNode)
 {
 	if(const FSchematicGraphNode* Node = FindNode(InChildNodeGuid))
 	{
-		return RemoveFromParentNode(Node);
+		return RemoveFromParentNode(Node, bUpdateGroupNode);
 	}
 	return false;
 }
 
-bool FSchematicGraphModel::RemoveFromParentNode(const FSchematicGraphNode* InChildNode)
+bool FSchematicGraphModel::RemoveFromParentNode(const FSchematicGraphNode* InChildNode, bool bUpdateGroupNode)
 {
 	check(InChildNode);
-	return SetParentNode(InChildNode, nullptr);
+	return SetParentNode(InChildNode, nullptr, bUpdateGroupNode);
 }
 
 FVector2d FSchematicGraphModel::GetPositionForNode(const FGuid& InNodeGuid) const
@@ -161,22 +186,29 @@ FVector2d FSchematicGraphModel::GetSizeForNode(const FSchematicGraphNode* InNode
 	return SSchematicGraphNode::DefaultNodeSize;
 }
 
-float FSchematicGraphModel::GetScaleForNode(const FGuid& InNodeGuid, bool bIncludeScaleOffset) const
+float FSchematicGraphModel::GetScaleForNode(const FGuid& InNodeGuid) const
 {
 	if(const FSchematicGraphNode* Node = FindNode(InNodeGuid))
 	{
-		return GetScaleForNode(Node, bIncludeScaleOffset);
+		return GetScaleForNode(Node);
 	}
 	return 1.f;
 }
 
-float FSchematicGraphModel::GetScaleForNode(const FSchematicGraphNode* InNode, bool bIncludeScaleOffset) const
+float FSchematicGraphModel::GetScaleForNode(const FSchematicGraphNode* InNode) const
 {
-	if(bIncludeScaleOffset)
+	float Scale = 1.f;
+
+	if(const FSchematicGraphNode* ParentNode = InNode->GetParentNode())
 	{
-		return GetScaleOffsetForNode(InNode);
+		const TOptional<float> ParentScale = GetScaleForChildNode(ParentNode, InNode);
+		if(ParentScale.IsSet())
+		{
+			Scale *= ParentScale.GetValue();
+		}
 	}
-	return 1.f;
+	
+	return Scale * GetScaleOffsetForNode(InNode);
 }
 
 float FSchematicGraphModel::GetScaleOffsetForNode(const FGuid& InNodeGuid) const
@@ -224,6 +256,10 @@ float FSchematicGraphModel::GetScaleOffsetForNode(const FSchematicGraphNode* InN
 bool FSchematicGraphModel::IsAutoScaleEnabledForNode(const FSchematicGraphNode* InNode) const
 {
 	check(InNode);
+	if(InNode->HasParentNode())
+	{
+		return false;
+	}
 	return InNode->IsAutoScaleEnabled();
 }
 
@@ -327,19 +363,80 @@ ESchematicGraphVisibility::Type FSchematicGraphModel::GetVisibilityForNode(const
 	return InNode->GetVisibility();
 }
 
-ESchematicGraphVisibility::Type FSchematicGraphModel::GetVisibilityForChildNodes(const FGuid& InNodeGuid) const
+TOptional<ESchematicGraphVisibility::Type> FSchematicGraphModel::GetVisibilityForChildNode(const FGuid& InParentNodeGuid, const FGuid& InChildNodeGuid) const
 {
-	if(const FSchematicGraphNode* Node = FindNode(InNodeGuid))
+	if(const FSchematicGraphNode* ParentNode = FindNode(InParentNodeGuid))
 	{
-		return GetVisibilityForChildNodes(Node);
+		if(const FSchematicGraphNode* ChildNode = FindNode(InChildNodeGuid))
+		{
+			return GetVisibilityForChildNode(ParentNode, ChildNode);
+		}
 	}
-	return ESchematicGraphVisibility::Visible;
+	return TOptional<ESchematicGraphVisibility::Type>();
 }
 
-ESchematicGraphVisibility::Type FSchematicGraphModel::GetVisibilityForChildNodes(const FSchematicGraphNode* InNode) const
+TOptional<ESchematicGraphVisibility::Type> FSchematicGraphModel::GetVisibilityForChildNode(const FSchematicGraphNode* InParentNode, const FSchematicGraphNode* InChildNode) const
 {
-	check(InNode);
-	return InNode->GetVisibilityForChildNodes();
+	check(InParentNode);
+	check(InChildNode);
+	return InParentNode->GetVisibilityForChildNode(InChildNode);
+}
+
+TOptional<FVector2d> FSchematicGraphModel::GetPositionForChildNode(const FGuid& InParentNodeGuid, const FGuid& InChildNodeGuid) const
+{
+	if(const FSchematicGraphNode* ParentNode = FindNode(InParentNodeGuid))
+	{
+		if(const FSchematicGraphNode* ChildNode = FindNode(InChildNodeGuid))
+		{
+			return GetPositionForChildNode(ParentNode, ChildNode);
+		}
+	}
+	return TOptional<FVector2d>();
+}
+
+TOptional<FVector2d> FSchematicGraphModel::GetPositionForChildNode(const FSchematicGraphNode* InParentNode, const FSchematicGraphNode* InChildNode) const
+{
+	check(InParentNode);
+	check(InChildNode);
+	return InParentNode->GetPositionForChildNode(InChildNode);
+}
+
+TOptional<float> FSchematicGraphModel::GetScaleForChildNode(const FGuid& InParentNodeGuid, const FGuid& InChildNodeGuid) const
+{
+	if(const FSchematicGraphNode* ParentNode = FindNode(InParentNodeGuid))
+	{
+		if(const FSchematicGraphNode* ChildNode = FindNode(InChildNodeGuid))
+		{
+			return GetScaleForChildNode(ParentNode, ChildNode);
+		}
+	}
+	return TOptional<float>();
+}
+
+TOptional<float> FSchematicGraphModel::GetScaleForChildNode(const FSchematicGraphNode* InParentNode, const FSchematicGraphNode* InChildNode) const
+{
+	check(InParentNode);
+	check(InChildNode);
+	return InParentNode->GetScaleForChildNode(InChildNode);
+}
+
+TOptional<bool> FSchematicGraphModel::GetInteractivityForChildNode(const FGuid& InParentNodeGuid, const FGuid& InChildNodeGuid) const
+{
+	if(const FSchematicGraphNode* ParentNode = FindNode(InParentNodeGuid))
+	{
+		if(const FSchematicGraphNode* ChildNode = FindNode(InChildNodeGuid))
+		{
+			return GetInteractivityForChildNode(ParentNode, ChildNode);
+		}
+	}
+	return TOptional<bool>();
+}
+
+TOptional<bool> FSchematicGraphModel::GetInteractivityForChildNode(const FSchematicGraphNode* InParentNode, const FSchematicGraphNode* InChildNode) const
+{
+	check(InParentNode);
+	check(InChildNode);
+	return InParentNode->GetInteractivityForChildNode(InChildNode);
 }
 
 bool FSchematicGraphModel::IsDragSupportedForNode(const FGuid& InNodeGuid) const
@@ -498,6 +595,18 @@ ESchematicGraphVisibility::Type FSchematicGraphModel::GetVisibilityForTag(const
 ESchematicGraphVisibility::Type FSchematicGraphModel::GetVisibilityForTag(const FSchematicGraphTag* InTag) const
 {
 	check(InTag);
+
+	if(InTag->IsA<FSchematicGraphGroupTag>())
+	{
+		if(const FSchematicGraphAutoGroupNode* AutoGroupNode = Cast<FSchematicGraphAutoGroupNode>(InTag->GetNode()))
+		{
+			if(AutoGroupNode->GetNumChildNodes() < 3 || AutoGroupNode->IsExpanded())
+			{
+				return ESchematicGraphVisibility::Hidden;
+			}
+		}
+	}
+	
 	return InTag->GetVisibility();
 }
 
@@ -726,6 +835,43 @@ ESchematicGraphVisibility::Type FSchematicGraphModel::GetVisibilityForLink(const
 	return InLink->GetVisibility();
 }
 
+FSchematicGraphGroupNode* FSchematicGraphModel::AddAutoGroupNode()
+{
+	return AddNode<FSchematicGraphAutoGroupNode>(); 
+}
+
+const FSchematicGraphGroupNode* FSchematicGraphModel::GetLastExpandedNode() const
+{
+	if(LastExpandedNode.IsValid())
+	{
+		return FindNode<FSchematicGraphGroupNode>(LastExpandedNode);
+	}
+	return nullptr;
+}
+
+void FSchematicGraphModel::SetLastExpandedNode(const FSchematicGraphGroupNode* InGroupNode)
+{
+	LastExpandedNode = InGroupNode ? InGroupNode->GetGuid() : FGuid();
+}
+
+void FSchematicGraphModel::Tick(float InDeltaTime)
+{
+	if(LastExpandedNode.IsValid())
+	{
+		if(const FSchematicGraphGroupNode* GroupNode = GetLastExpandedNode())
+		{
+			if(!GroupNode->IsExpanded())
+			{
+				LastExpandedNode = FGuid();
+			}
+		}
+		else
+		{
+			LastExpandedNode = FGuid();
+		}
+	}
+}
+
 #undef LOCTEXT_NAMESPACE
 
 #endif
\ No newline at end of file
diff --git a/Engine/Source/Editor/AnimationEditorWidgets/Private/SchematicGraphPanel/SchematicGraphNode.cpp b/Engine/Source/Editor/AnimationEditorWidgets/Private/SchematicGraphPanel/SchematicGraphNode.cpp
--- a/Engine/Source/Editor/AnimationEditorWidgets/Private/SchematicGraphPanel/SchematicGraphNode.cpp
+++ b/Engine/Source/Editor/AnimationEditorWidgets/Private/SchematicGraphPanel/SchematicGraphNode.cpp
@@ -4,9 +4,20 @@
 
 #include "SchematicGraphPanel/SchematicGraphNode.h"
 #include "SchematicGraphPanel/SchematicGraphModel.h"
+#include <SchematicGraphPanel/SchematicGraphStyle.h>
+#include "Framework/Application/SlateApplication.h"
 
 #define LOCTEXT_NAMESPACE "SchematicGraphNode"
 
+FSchematicGraphNode::FSchematicGraphNode()
+{
+	static const FSlateBrush* BackgroundBrush = FSchematicGraphStyle::Get().GetBrush( "Schematic.Background");
+	static const FSlateBrush* OutlineBrush = FSchematicGraphStyle::Get().GetBrush( "Schematic.Outline.Single");
+	static const FSlateBrush* DotBrush = FSchematicGraphStyle::Get().GetBrush( "Schematic.Dot.Small");
+	Brushes = { BackgroundBrush, OutlineBrush, DotBrush };
+	Colors = { FLinearColor::White * 0.4f, FLinearColor::White, FLinearColor::Blue };
+}
+
 const FSchematicGraphNode* FSchematicGraphNode::GetParentNode() const
 {
 	if(Model && ParentNodeGuid.IsValid())
@@ -16,6 +27,29 @@ const FSchematicGraphNode* FSchematicGraphNode::GetParentNode() const
 	return nullptr;
 }
 
+const FGuid& FSchematicGraphNode::GetRootNodeGuid() const
+{
+	if(const FSchematicGraphNode* ParentNode = GetParentNode())
+	{
+		return ParentNode->GetRootNodeGuid();
+	}
+	return GetGuid();
+}
+
+const FSchematicGraphNode* FSchematicGraphNode::GetRootNode() const
+{
+	if(Model)
+	{
+		return Model->FindNode(GetRootNodeGuid());
+	}
+	return nullptr;
+}
+
+const FSchematicGraphGroupNode* FSchematicGraphNode::GetGroupNode() const
+{
+	return Cast<FSchematicGraphGroupNode>(GetParentNode());
+}
+
 const FSchematicGraphNode* FSchematicGraphNode::GetChildNode(int32 InChildNodeIndex) const
 {
 	check(ChildNodeGuids.IsValidIndex(InChildNodeIndex));
@@ -28,24 +62,153 @@ const FSchematicGraphNode* FSchematicGraphNode::GetChildNode(int32 InChildNodeIn
 	return nullptr;
 }
 
+TOptional<ESchematicGraphVisibility::Type> FSchematicGraphNode::GetVisibilityForChildNode(const FGuid& InChildGuid) const
+{
+	if(Model)
+	{
+		if(const FSchematicGraphNode* ChildNode = Model->FindNode(InChildGuid))
+		{
+			return GetVisibilityForChildNode(ChildNode);
+		}
+	}
+	return TOptional<ESchematicGraphVisibility::Type>();
+}
+
+TOptional<FVector2d> FSchematicGraphNode::GetPositionForChildNode(const FGuid& InChildGuid) const
+{
+	if(Model)
+	{
+		if(const FSchematicGraphNode* ChildNode = Model->FindNode(InChildGuid))
+		{
+			return GetPositionForChildNode(ChildNode);
+		}
+	}
+	return TOptional<FVector2d>();
+}
+
+TOptional<float> FSchematicGraphNode::GetScaleForChildNode(const FGuid& InChildGuid) const
+{
+	if(Model)
+	{
+		if(const FSchematicGraphNode* ChildNode = Model->FindNode(InChildGuid))
+		{
+			return GetScaleForChildNode(ChildNode);
+		}
+	}
+	return TOptional<float>();
+}
+
+TOptional<bool> FSchematicGraphNode::GetInteractivityForChildNode(const FGuid& InChildGuid) const
+{
+	if(Model)
+	{
+		if(const FSchematicGraphNode* ChildNode = Model->FindNode(InChildGuid))
+		{
+			return GetInteractivityForChildNode(ChildNode);
+		}
+	}
+	return TOptional<bool>();
+}
+
+FVector2d FSchematicGraphNode::GetPosition() const
+{
+	if(const FSchematicGraphNode* ParentNode = GetParentNode())
+	{
+		const TOptional<FVector2d> ParentPosition = Model->GetPositionForChildNode(ParentNode, this);
+		if(ParentPosition.IsSet())
+		{
+			return ParentPosition.GetValue();
+		}
+	}
+	return Position;
+}
+
+const FText& FSchematicGraphNode::GetToolTip() const
+{
+	return ToolTip;
+}
+
 ESchematicGraphVisibility::Type FSchematicGraphNode::GetVisibility() const
 {
 	if(Visibility == ESchematicGraphVisibility::Hidden)
 	{
 		return Visibility;
 	}
+
+	const FSchematicGraphNode* ParentNode = GetParentNode();
 	
-	if(const FSchematicGraphNode* ParentNode = GetParentNode())
+	if(const FSchematicGraphGroupNode* LastExpandedNode = Model->GetLastExpandedNode())
 	{
-		const ESchematicGraphVisibility::Type ParentVisibility = Model->GetVisibilityForChildNodes(ParentNode);
-		if(ParentVisibility != ESchematicGraphVisibility::Visible)
+		if(LastExpandedNode != this && LastExpandedNode != ParentNode)
 		{
-			return ParentVisibility;
+			return ESchematicGraphVisibility::FadedOut;
+		}
+	}
+	
+	if(ParentNode)
+	{
+		const TOptional<ESchematicGraphVisibility::Type> ParentVisibility = Model->GetVisibilityForChildNode(ParentNode, this);
+		if(ParentVisibility.IsSet())
+		{
+			return ParentVisibility.GetValue();
 		}
 	}
 	return Visibility;
 }
 
+bool FSchematicGraphNode::IsInteractive() const
+{
+	if(GetVisibility() != ESchematicGraphVisibility::Visible)
+	{
+		return false;
+	}
+	
+	if(const FSchematicGraphNode* ParentNode = GetParentNode())
+	{
+		const TOptional<bool> ParentInteractivity = Model->GetInteractivityForChildNode(ParentNode, this);
+		if(ParentInteractivity.IsSet())
+		{
+			return ParentInteractivity.GetValue();
+		}
+	}
+	return true;
+}
+
+FReply FSchematicGraphNode::OnClicked(const FPointerEvent& InMouseEvent)
+{
+	return FReply::Unhandled();
+}
+
+void FSchematicGraphNode::OnMouseEnter()
+{
+	SetScaleOffset(ScaledUp);
+}
+
+void FSchematicGraphNode::OnMouseLeave()
+{
+	SetScaleOffset(1.f);
+}
+
+void FSchematicGraphNode::OnDragOver()
+{
+	SetScaleOffset(ScaledDown);
+
+	if(FSchematicGraphGroupNode* GroupNode = const_cast<FSchematicGraphGroupNode*>(GetGroupNode()))
+	{
+		GroupNode->SetExpanded(true);
+	}
+}
+
+void FSchematicGraphNode::OnDragLeave()
+{
+	SetScaleOffset(1.f);
+
+	if(FSchematicGraphGroupNode* GroupNode = const_cast<FSchematicGraphGroupNode*>(GetGroupNode()))
+	{
+		GroupNode->SetExpanded(false);
+	}
+}
+
 FString FSchematicGraphNode::GetDragDropDecoratorLabel() const
 {
 	return GetGuid().ToString();
@@ -78,13 +241,205 @@ void FSchematicGraphNode::NotifyTagAdded(const TSharedPtr<FSchematicGraphTag>& T
 	}
 }
 
-ESchematicGraphVisibility::Type FSchematicGraphGroupNode::GetVisibilityForChildNodes() const
+FSchematicGraphGroupNode::FSchematicGraphGroupNode()
+	: AnimationSettings(EEasingInterpolatorType::CubicEaseOut, 0.35f)
+{
+	static const FSlateBrush* BackgroundBrush = FSchematicGraphStyle::Get().GetBrush( "Schematic.Background");
+	static const FSlateBrush* OutlineBrush = FSchematicGraphStyle::Get().GetBrush( "Schematic.Outline.Single");
+	static const FSlateBrush* GroupBrush = FSchematicGraphStyle::Get().GetBrush( "Schematic.Dot.Group");
+	Brushes = { BackgroundBrush, OutlineBrush, GroupBrush };
+	Colors = { FLinearColor::White * 0.4f, FLinearColor::White, FLinearColor::Gray };
+
+	ExpansionState = TAnimatedAttribute<float>::Create(AnimationSettings, 0.f);
+	(void)AddTag<FSchematicGraphGroupTag>();
+}
+
+bool FSchematicGraphGroupNode::IsExpanded() const
+{
+	return GetExpansionState() > SMALL_NUMBER;
+}
+
+bool FSchematicGraphGroupNode::IsExpanding() const
+{
+	return (ExpansionState->Get() < 1.f - SMALL_NUMBER) &&
+		(ExpansionState->GetDesiredValue() >= 1.f - SMALL_NUMBER);
+}
+
+bool FSchematicGraphGroupNode::IsCollapsing() const
+{
+	return (ExpansionState->Get() >= SMALL_NUMBER) &&
+		(ExpansionState->GetDesiredValue() < SMALL_NUMBER);
+}
+
+float FSchematicGraphGroupNode::GetExpansionState() const
+{
+	if(!ExpansionState.IsValid())
+	{
+		return 0.f;
+	}
+	return ExpansionState->Get();
+}
+
+void FSchematicGraphGroupNode::SetExpanded(bool InExpanded, bool bAutoCloseParentGroups)
+{
+	if(InExpanded)
+	{
+		if(GetNumChildNodes() == 0)
+		{
+			return;
+		}
+	}
+	
+	if(!InExpanded && GetExpansionState() < SMALL_NUMBER)
+	{
+		ExpansionState->SetValueAndStop(0.f);
+	}
+	else
+	{
+		if(InExpanded && !ExpansionState->IsPlaying())
+		{
+			// set the state to an initial value so that IsExpanded will return correctly
+			if(ExpansionState->Get() < SMALL_NUMBER)
+			{
+				ExpansionState->SetValueAndStop(SMALL_NUMBER * 2.f);
+			}
+		}
+		if(!ExpansionState->GetDelay().IsSet())
+		{
+			ExpansionState->SetDelayOneShot(GetDelayDuration(InExpanded));
+		}
+		ExpansionState->Set(InExpanded ? 1.f : 0.f);
+	}
+
+	if(InExpanded || bAutoCloseParentGroups)
+	{
+		if(FSchematicGraphGroupNode* GroupNode = const_cast<FSchematicGraphGroupNode*>(GetGroupNode()))
+		{
+			GroupNode->SetExpanded(InExpanded, bAutoCloseParentGroups);
+		}
+	}
+
+	if(InExpanded)
+	{
+		Model->SetLastExpandedNode(this);
+	}
+	else if(const FSchematicGraphGroupNode* GroupNode = GetGroupNode())
+	{
+		if(GroupNode->IsExpanded())
+		{
+			Model->SetLastExpandedNode(GroupNode);
+		}
+	}
+}
+
+TOptional<ESchematicGraphVisibility::Type> FSchematicGraphGroupNode::GetVisibilityForChildNode(const FSchematicGraphNode* InChildNode) const
 {
 	if(!IsExpanded())
 	{
 		return ESchematicGraphVisibility::Hidden;
 	}
-	return FSchematicGraphNode::GetVisibilityForChildNodes();
+	return Super::GetVisibilityForChildNode(InChildNode);
+}
+
+TOptional<FVector2d> FSchematicGraphGroupNode::GetPositionForChildNode(const FSchematicGraphNode* InChildNode) const
+{
+	if(IsExpanded() && Model)
+	{
+		const int32 ChildNodeIndex = ChildNodeGuids.Find(InChildNode->GetGuid());
+		if(ChildNodeIndex != INDEX_NONE)
+		{
+			const float Distance = ExpansionRadius * GetExpansionState();
+			const float Angle = 360.f * float(ChildNodeIndex) / float(ChildNodeGuids.Num());
+			const FVector2d Offset = FVector2d(Distance, 0).GetRotated(Angle);
+			return Model->GetPositionForNode(this) + Offset;
+		}
+	}
+	return Super::GetPositionForChildNode(InChildNode);
+}
+
+TOptional<float> FSchematicGraphGroupNode::GetScaleForChildNode(const FSchematicGraphNode* InChildNode) const
+{
+	if(IsExpanded() && Model)
+	{
+		return GetExpansionState();
+	}
+	return Super::GetScaleForChildNode(InChildNode);
+}
+
+TOptional<bool> FSchematicGraphGroupNode::GetInteractivityForChildNode(const FSchematicGraphNode* InChildNode) const
+{
+	if(GetExpansionState() < 1.f - SMALL_NUMBER)
+	{
+		return false;
+	}
+	return Super::GetInteractivityForChildNode(InChildNode);
+}
+
+FReply FSchematicGraphGroupNode::OnClicked(const FPointerEvent& InMouseEvent)
+{
+	if(GetNumChildNodes() > 0)
+	{
+		SetExpanded(!IsExpanded());
+		return FReply::Handled();
+	}
+	return FSchematicGraphNode::OnClicked(InMouseEvent);
+}
+
+void FSchematicGraphGroupNode::OnDragOver()
+{
+	Super::OnDragOver();
+	SetExpanded(true, true);
+}
+
+void FSchematicGraphGroupNode::OnDragLeave()
+{
+	Super::OnDragLeave();
+	SetExpanded(false, true);
+}
+
+float FSchematicGraphGroupNode::GetDelayDuration(bool bEnter) const
+{
+	if (FSlateApplication::Get().IsDragDropping())
+	{
+		return bEnter ? EnterDelayDuration : LeaveDelayDuration;
+	}
+	return 0.f;
+}
+
+void FSchematicGraphGroupNode::SetAnimationSettings(const TEasingAttributeInterpolator<float>::FSettings& InSettings)
+{
+	ExpansionState = TAnimatedAttribute<float>::Create(AnimationSettings, IsExpanded() ? 1.f : 0.f);
+}
+
+const FText& FSchematicGraphAutoGroupNode::GetLabel() const
+{
+	const FText& SuperLabel = FSchematicGraphGroupNode::GetLabel();
+	if(SuperLabel.IsEmpty() && !IsExpanded())
+	{
+		AutoGroupLabel = FText();
+
+		for(int32 Index = 0; Index < GetNumChildNodes(); Index++)
+		{
+			if(const FSchematicGraphNode* ChildNode = GetChildNode(Index))
+			{
+				const FText& ChildLabel = ChildNode->GetLabel();
+				if(!ChildLabel.IsEmpty())
+				{
+					if(!AutoGroupLabel.IsEmpty())
+					{
+						AutoGroupLabel = FText::Format(LOCTEXT("AutoGroupLabelAppendFormat", "{0}\n{1}"), AutoGroupLabel, ChildLabel);
+					}
+					else
+					{
+						AutoGroupLabel = ChildLabel;
+					}
+				}
+			}
+		}
+
+		return AutoGroupLabel;
+	}
+	return SuperLabel;
 }
 
 #undef LOCTEXT_NAMESPACE
diff --git a/Engine/Source/Editor/AnimationEditorWidgets/Private/SchematicGraphPanel/SchematicGraphTag.cpp b/Engine/Source/Editor/AnimationEditorWidgets/Private/SchematicGraphPanel/SchematicGraphTag.cpp
--- a/Engine/Source/Editor/AnimationEditorWidgets/Private/SchematicGraphPanel/SchematicGraphTag.cpp
+++ b/Engine/Source/Editor/AnimationEditorWidgets/Private/SchematicGraphPanel/SchematicGraphTag.cpp
@@ -4,16 +4,37 @@
 
 #include "SchematicGraphPanel/SchematicGraphTag.h"
 #include "SchematicGraphPanel/SchematicGraphNode.h"
+#include "SchematicGraphPanel/SchematicGraphStyle.h"
 
 #define LOCTEXT_NAMESPACE "SchematicGraphTag"
 
+FSchematicGraphTag::FSchematicGraphTag()
+{
+	static const FSlateBrush* CircleBrush = FSchematicGraphStyle::Get().GetBrush( "Schematic.Background");
+	BackgroundBrush = CircleBrush;
+}
+
 const FSchematicGraphNode* FSchematicGraphTag::GetNode() const
 {
-	if(Node.IsValid())
+	return Node;
+}
+
+ESchematicGraphVisibility::Type FSchematicGraphGroupTag::GetVisibility() const
+{
+	if(const FSchematicGraphNode* GroupNode = Cast<FSchematicGraphNode>(GetNode()))
+	{
+		return GroupNode->GetNumChildNodes() == 0 ? ESchematicGraphVisibility::Hidden : ESchematicGraphVisibility::Visible;
+	}
+	return Super::GetVisibility();
+}
+
+FText FSchematicGraphGroupTag::GetLabel() const
+{
+	if(const FSchematicGraphNode* GroupNode = Cast<FSchematicGraphNode>(GetNode()))
 	{
-		return Node.Pin().Get();
+		return FText::FromString(FString::FromInt(GroupNode->GetNumChildNodes()));
 	}
-	return nullptr;
+	return Super::GetLabel();
 }
 
 #undef LOCTEXT_NAMESPACE
diff --git a/Engine/Source/Editor/AnimationEditorWidgets/Public/SchematicGraphPanel/SSchematicGraphPanel.h b/Engine/Source/Editor/AnimationEditorWidgets/Public/SchematicGraphPanel/SSchematicGraphPanel.h
--- a/Engine/Source/Editor/AnimationEditorWidgets/Public/SchematicGraphPanel/SSchematicGraphPanel.h
+++ b/Engine/Source/Editor/AnimationEditorWidgets/Public/SchematicGraphPanel/SSchematicGraphPanel.h
@@ -58,7 +58,7 @@ public:
 	typedef TAnimatedAttribute<float> FFloatAttribute;
 	typedef TAnimatedAttribute<FLinearColor> FLinearColorAttribute;
 	
-	DECLARE_DELEGATE_OneParam(FOnClicked, SSchematicGraphNode*);
+	DECLARE_DELEGATE_TwoParams(FOnClicked, SSchematicGraphNode*, const FPointerEvent&);
 	DECLARE_DELEGATE_TwoParams(FOnBeginDrag, SSchematicGraphNode*, const FDragDropOperation&);
 	DECLARE_DELEGATE_TwoParams(FOnEndDrag, SSchematicGraphNode*, const FDragDropOperation&);
 	DECLARE_DELEGATE_TwoParams(FOnDrop, SSchematicGraphNode*, const FDragDropEvent&);
@@ -98,18 +98,18 @@ public:
 	const FVector2d& GetOriginalSize() const { return OriginalSize; }
 
 	const FSchematicGraphNode* GetNodeData() const { return NodeData; }
+	FSchematicGraphNode* GetNodeData() { return NodeData; }
 	const FGuid GetGuid() const;
 	bool IsInteractive() const;
 
-	const bool IsBeingDragged() const { return bIsBeingDragged; } 
+	const bool IsBeingDragged() const { return bIsBeingDragged; }
+	const bool IsFadedOut() const;
 	
 private:
 
 	bool bIsBeingDragged = false;
 	static inline const FVector2d DefaultNodeSize = FVector2d(32.0,32.0);  
 	FVector2d OriginalSize = DefaultNodeSize;
-	static inline constexpr float ScaledUp = 1.25;
-	static inline constexpr float ScaledDown = 0.75;
 
 	FSchematicGraphNode* NodeData = nullptr;
 	TSharedPtr<FVector2dAttribute> Position;
@@ -126,6 +126,7 @@ private:
 	FOnEndDrag OnEndDragDelegate;
 	FOnDrop OnDropDelegate;
 	SSchematicGraphPanel* SchematicGraphPanel = nullptr;
+	TSharedPtr<FFloatAttribute> ExpansionCircleFactor;
 
 	friend class SSchematicGraphPanel;
 	friend class FSchematicGraphModel;
@@ -141,7 +142,7 @@ public:
 	using FFloatAttribute = SSchematicGraphNode::FFloatAttribute;
 	using FLinearColorAttribute = SSchematicGraphNode::FLinearColorAttribute;
 
-	DECLARE_DELEGATE_TwoParams(FOnNodeClicked, SSchematicGraphPanel*, SSchematicGraphNode*);
+	DECLARE_DELEGATE_ThreeParams(FOnNodeClicked, SSchematicGraphPanel*, SSchematicGraphNode*, const FPointerEvent&);
 	DECLARE_DELEGATE_ThreeParams(FOnBeginDrag, SSchematicGraphPanel*, SSchematicGraphNode*, const FDragDropOperation&);
 	DECLARE_DELEGATE_ThreeParams(FOnEndDrag, SSchematicGraphPanel*, SSchematicGraphNode*, const FDragDropOperation&);
 	DECLARE_DELEGATE_ThreeParams(FOnDrop, SSchematicGraphPanel*, SSchematicGraphNode*, const FDragDropEvent&);
@@ -175,16 +176,20 @@ public:
 	}
 
 	void SetSchematicGraph(FSchematicGraphModel* InGraphData);
+	const FSchematicGraphModel* GetSchematicGraph() const { return GraphData; }
+	FSchematicGraphModel* GetSchematicGraph() { return GraphData; }
 	void Construct(const FArguments& InArgs);
 
 	void RebuildPanel();
 	void AddNode(const FSchematicGraphNode* InNodeToAdd);
 	void RemoveNode(const FSchematicGraphNode* InNodeToRemove);
 	const SSchematicGraphNode* FindNode(const FGuid& InGuid) const;
+	SSchematicGraphNode* FindNode(const FGuid& InGuid);
 	void AddLink(const FSchematicGraphLink* InLinkToAdd);
 	void RemoveLink(const FSchematicGraphLink* InLinkToRemove);
 	const FSchematicLinkWidgetInfo* FindLink(const FGuid& InGuid) const;
-	
+	FSchematicLinkWidgetInfo* FindLink(const FGuid& InGuid);
+
 	// SNodePanel interface
 	virtual void OnArrangeChildren(const FGeometry& AllottedGeometry, FArrangedChildren& ArrangedChildren) const override;
 	virtual int32 OnPaint(const FPaintArgs& Args, const FGeometry& AllottedGeometry, const FSlateRect& MyCullingRect, FSlateWindowElementList& OutDrawElements, int32 LayerId, const FWidgetStyle& InWidgetStyle, bool bParentEnabled) const override;
@@ -201,7 +206,7 @@ public:
 	virtual bool IsTickable() const override { return true; }
 	// End of FTickableEditorObject interface
 
-	void OnNodeClicked(SSchematicGraphNode* Node);
+	void OnNodeClicked(SSchematicGraphNode* Node, const FPointerEvent& MouseEvent);
 	void OnBeginDragEvent(SSchematicGraphNode* Node, const FDragDropOperation& InDragDropEvent);
 	void OnEndDragEvent(SSchematicGraphNode* Node, const FDragDropOperation& InDragDropEvent);
 	void OnDropEvent(SSchematicGraphNode* Node, const FDragDropEvent& InDragDropEvent);
@@ -209,21 +214,34 @@ public:
 
 	virtual FVector2d GetPositionForNode(FGuid InNodeGuid) const;
 	virtual FLinearColor GetColorForNode(FGuid InNodeGuid, int32 InLayerIndex) const;
-	virtual float GetScaleForNode(FGuid InNodeGuid, bool bIncludeScaleOffset) const;
+	virtual float GetScaleForNode(FGuid InNodeGuid) const;
+
+	virtual bool IsAutoGroupingEnabled() const;
+	virtual float GetAutoGroupingDistance() const;
+
 	virtual bool IsAutoScaleEnabledForNode(FGuid InNodeGuid) const;
 	virtual float GetMinimumLinkDistanceForNode(FGuid InLinkGuid, bool bIncludeScale = true) const;
 
+	FOnNodeClicked& OnNodeClicked() { return OnNodeClickedDelegate; }
+	FOnBeginDrag& OnBeginDrag() { return OnBeginDragDelegate; }
+	FOnEndDrag& OnEndDrag() { return OnEndDragDelegate; }
+	FOnDrop& OnAcceptDrop() { return OnDropDelegate; }
+
+private:
+
+	void UpdatePerNodeCaches(bool bRemoveNodesFromAutoGroups);
+	void UpdateAutoGroupingForNodes();
 	void UpdateAutoScalingForNodes();
 
 	bool bIsDragDropping = false;
-	bool bIsOverlay;
+	bool bIsOverlay = false;
 	int32 PaddingLeft = 0;
 	int32 PaddingRight = 0;
 	int32 PaddingTop = 0;
 	int32 PaddingBottom = 0;
 	int32 PaddingInterNode = 0;
 	TSharedPtr<FFloatAttribute> FadeBackgroundAlpha;
-	FSchematicGraphModel* GraphData;
+	FSchematicGraphModel* GraphData = nullptr;
 	FOnNodeClicked OnNodeClickedDelegate;
 	FOnBeginDrag OnBeginDragDelegate;
 	FOnEndDrag OnEndDragDelegate;
@@ -234,11 +252,31 @@ public:
 	TArray<FGuid> NodesBottomRight;
 	TMap<FGuid, TSharedPtr<SSchematicGraphNode>> NodeByGuid;
 
+	struct FPerNodeCache
+	{
+		FPerNodeCache()
+			: bHasParent(false)
+			, Visibility(ESchematicGraphVisibility::Visible)
+			, bIsAutoScaling(false)
+			, Position(FVector2d::ZeroVector)
+			, Radius(0.0)
+		{}
+
+		bool bHasParent;
+		ESchematicGraphVisibility::Type Visibility;
+		bool bIsAutoScaling;
+		FVector2d Position;
+		double Radius;
+	};
+
+	TArray<FPerNodeCache> PerNodeCaches;
+	TMap<uint32, FGuid> GroupNodeGuidByHash;
+
 	TMap<FGuid, TSharedPtr<FSchematicLinkWidgetInfo>> LinkByGuid;
 
 	mutable TMap<FGuid, FVector2d> NodeCenterByGuid;
-	mutable TArray<bool> NodeVisibilityByIndex;
-	mutable TMap<FGuid, bool> NodeVisibilityByGuid;
+	mutable TArray<ESchematicGraphVisibility::Type> NodeVisibilityByIndex;
+	mutable TMap<FGuid, ESchematicGraphVisibility::Type> NodeVisibilityByGuid;
 };
 
 #endif
\ No newline at end of file
diff --git a/Engine/Source/Editor/AnimationEditorWidgets/Public/SchematicGraphPanel/SchematicGraphDefines.h b/Engine/Source/Editor/AnimationEditorWidgets/Public/SchematicGraphPanel/SchematicGraphDefines.h
--- a/Engine/Source/Editor/AnimationEditorWidgets/Public/SchematicGraphPanel/SchematicGraphDefines.h
+++ b/Engine/Source/Editor/AnimationEditorWidgets/Public/SchematicGraphPanel/SchematicGraphDefines.h
@@ -6,6 +6,8 @@
 #include "CoreMinimal.h"
 #include "Styling/SlateBrush.h"
 
+ANIMATIONEDITORWIDGETS_API DECLARE_LOG_CATEGORY_EXTERN(LogSchematicGraph, Log, All);
+
 class FSchematicGraphTag;
 class FSchematicGraphNode;
 class FSchematicGraphLink;
@@ -101,6 +103,7 @@ friend T* CastChecked(ClassName* InElement) \
 }
 
 #define SCHEMATICGRAPHELEMENT_BODY(ClassName, SuperClass, BaseClass) \
+typedef SuperClass Super; \
 inline static const FName& Type = TEXT(#ClassName); \
 virtual const FName& GetType() const override { return ClassName::Type; } \
 virtual bool IsA(const FName& InType) const \
diff --git a/Engine/Source/Editor/AnimationEditorWidgets/Public/SchematicGraphPanel/SchematicGraphLink.h b/Engine/Source/Editor/AnimationEditorWidgets/Public/SchematicGraphPanel/SchematicGraphLink.h
--- a/Engine/Source/Editor/AnimationEditorWidgets/Public/SchematicGraphPanel/SchematicGraphLink.h
+++ b/Engine/Source/Editor/AnimationEditorWidgets/Public/SchematicGraphPanel/SchematicGraphLink.h
@@ -8,7 +8,7 @@
 #define SCHEMATICGRAPHLINK_BODY(ClassName, SuperClass) \
 SCHEMATICGRAPHELEMENT_BODY(ClassName, SuperClass, FSchematicGraphNode)
 
-class ANIMATIONEDITORWIDGETS_API FSchematicGraphLink
+class ANIMATIONEDITORWIDGETS_API FSchematicGraphLink : public TSharedFromThis<FSchematicGraphLink>
 {
 public:
 
diff --git a/Engine/Source/Editor/AnimationEditorWidgets/Public/SchematicGraphPanel/SchematicGraphModel.h b/Engine/Source/Editor/AnimationEditorWidgets/Public/SchematicGraphPanel/SchematicGraphModel.h
--- a/Engine/Source/Editor/AnimationEditorWidgets/Public/SchematicGraphPanel/SchematicGraphModel.h
+++ b/Engine/Source/Editor/AnimationEditorWidgets/Public/SchematicGraphPanel/SchematicGraphModel.h
@@ -7,11 +7,11 @@
 #include "SchematicGraphNode.h"
 #include "SchematicGraphLink.h"
 
-class ANIMATIONEDITORWIDGETS_API FSchematicGraphModel
+class ANIMATIONEDITORWIDGETS_API FSchematicGraphModel : public TSharedFromThis<FSchematicGraphModel>
 {
 public:
 
-	virtual ~FSchematicGraphModel() {}
+	SCHEMATICGRAPHELEMENT_BODY_BASE(FSchematicGraphModel)
 	
 	virtual void Reset();
 	
@@ -40,6 +40,13 @@ public:
 		return nullptr;
 	}
 
+	template<typename NodeType = FSchematicGraphNode>
+	NodeType* FindNode(const FGuid& InNodeGuid)
+	{
+		const FSchematicGraphModel* ConstThis = this;
+		return const_cast<NodeType*>(ConstThis->FindNode<NodeType>(InNodeGuid));
+	}
+
 	template<typename NodeType = FSchematicGraphNode>
 	const NodeType* FindNodeChecked(const FGuid& InNodeGuid) const
 	{
@@ -51,12 +58,19 @@ public:
 		return nullptr;
 	}
 
+	template<typename NodeType = FSchematicGraphNode>
+	NodeType* FindNodeChecked(const FGuid& InNodeGuid)
+	{
+		const FSchematicGraphModel* ConstThis = this;
+		return const_cast<NodeType*>(ConstThis->FindNodeChecked<NodeType>(InNodeGuid));
+	}
+
 	virtual bool RemoveNode(const FGuid& InNodeGuid);
 
-	virtual bool SetParentNode(const FGuid& InChildNodeGuid, const FGuid& InParentNodeGuid);
-	virtual bool SetParentNode(const FSchematicGraphNode* InChildNode, const FSchematicGraphNode* InParentNode);
-	bool RemoveFromParentNode(const FGuid& InChildNodeGuid);
-	bool RemoveFromParentNode(const FSchematicGraphNode* InChildNode);
+	virtual bool SetParentNode(const FGuid& InChildNodeGuid, const FGuid& InParentNodeGuid, bool bUpdateGroupNode = true);
+	virtual bool SetParentNode(const FSchematicGraphNode* InChildNode, const FSchematicGraphNode* InParentNode, bool bUpdateGroupNode = true);
+	bool RemoveFromParentNode(const FGuid& InChildNodeGuid, bool bUpdateGroupNode = true);
+	bool RemoveFromParentNode(const FSchematicGraphNode* InChildNode, bool bUpdateGroupNode = true);
 
 	const TArray<TSharedPtr<FSchematicGraphNode>>& GetNodes() const { return Nodes; }
 
@@ -68,8 +82,8 @@ public:
 	virtual bool GetPositionAnimationEnabledForNode(const FSchematicGraphNode* InNode) const;
 	FVector2d GetSizeForNode(const FGuid& InNodeGuid) const;
 	virtual FVector2d GetSizeForNode(const FSchematicGraphNode* InNode) const;
-	float GetScaleForNode(const FGuid& InNodeGuid, bool bIncludeScaleOffset) const;
-	virtual float GetScaleForNode(const FSchematicGraphNode* InNode, bool bIncludeScaleOffset) const;
+	float GetScaleForNode(const FGuid& InNodeGuid) const;
+	virtual float GetScaleForNode(const FSchematicGraphNode* InNode) const;
 	float GetScaleOffsetForNode(const FGuid& InNodeGuid) const;
 	virtual float GetScaleOffsetForNode(const FSchematicGraphNode* InNode) const;
 	float GetMinimumLinkDistanceForNode(const FGuid& InNodeGuid) const;
@@ -88,8 +102,14 @@ public:
 	virtual ESchematicGraphPlacementConstraint::Type GetPlacementForNode(const FSchematicGraphNode* InNode) const;
 	ESchematicGraphVisibility::Type GetVisibilityForNode(const FGuid& InNodeGuid) const;
 	virtual ESchematicGraphVisibility::Type GetVisibilityForNode(const FSchematicGraphNode* InNode) const;
-	ESchematicGraphVisibility::Type GetVisibilityForChildNodes(const FGuid& InNodeGuid) const;
-	virtual ESchematicGraphVisibility::Type GetVisibilityForChildNodes(const FSchematicGraphNode* InNode) const;
+	TOptional<ESchematicGraphVisibility::Type> GetVisibilityForChildNode(const FGuid& InParentNodeGuid, const FGuid& InChildNodeGuid) const;
+	virtual TOptional<ESchematicGraphVisibility::Type> GetVisibilityForChildNode(const FSchematicGraphNode* InParentNode, const FSchematicGraphNode* InChildNode) const;
+	TOptional<FVector2d> GetPositionForChildNode(const FGuid& InParentNodeGuid, const FGuid& InChildNodeGuid) const;
+	virtual TOptional<FVector2d> GetPositionForChildNode(const FSchematicGraphNode* InParentNode, const FSchematicGraphNode* InChildNode) const;
+	TOptional<float> GetScaleForChildNode(const FGuid& InParentNodeGuid, const FGuid& InChildNodeGuid) const;
+	virtual TOptional<float> GetScaleForChildNode(const FSchematicGraphNode* InParentNode, const FSchematicGraphNode* InChildNode) const;
+	TOptional<bool> GetInteractivityForChildNode(const FGuid& InParentNodeGuid, const FGuid& InChildNodeGuid) const;
+	virtual TOptional<bool> GetInteractivityForChildNode(const FSchematicGraphNode* InParentNode, const FSchematicGraphNode* InChildNode) const;
 	bool IsDragSupportedForNode(const FGuid& InNodeGuid) const;
 	virtual bool IsDragSupportedForNode(const FSchematicGraphNode* InNode) const;
 
@@ -140,6 +160,13 @@ public:
 		return nullptr;
 	}
 
+	template<typename LinkType = FSchematicGraphLink>
+	LinkType* FindLink(const FGuid& InLinkGuid)
+	{
+		const FSchematicGraphModel* ConstThis = this;
+		return const_cast<LinkType*>(ConstThis->FindLink<LinkType>(InLinkGuid));
+	}
+
 	template<typename LinkType = FSchematicGraphLink>
 	const LinkType* FindLinkChecked(const FGuid& InLinkGuid) const
 	{
@@ -151,6 +178,13 @@ public:
 		return nullptr;
 	}
 
+	template<typename LinkType = FSchematicGraphLink>
+	LinkType* FindLinkChecked(const FGuid& InLinkGuid)
+	{
+		const FSchematicGraphModel* ConstThis = this;
+		return const_cast<LinkType*>(ConstThis->FindLinkChecked<LinkType>(InLinkGuid));
+	}
+	
 	template<typename LinkType = FSchematicGraphLink>
 	const LinkType* FindLink(const FGuid& InSourceNodeGuid, const FGuid& InTargetNodeGuid) const
 	{
@@ -162,6 +196,13 @@ public:
 		return nullptr;
 	}
 
+	template<typename LinkType = FSchematicGraphLink>
+	LinkType* FindLink(const FGuid& InSourceNodeGuid, const FGuid& InTargetNodeGuid)
+	{
+		const FSchematicGraphModel* ConstThis = this;
+		return const_cast<LinkType*>(ConstThis->FindLink<LinkType>(InSourceNodeGuid, InTargetNodeGuid));
+	}
+
 	template<typename LinkType = FSchematicGraphLink>
 	const LinkType* FindLinkChecked(const FGuid& InSourceNodeGuid, const FGuid& InTargetNodeGuid) const
 	{
@@ -174,6 +215,13 @@ public:
 		return nullptr;
 	}
 
+	template<typename LinkType = FSchematicGraphLink>
+	LinkType* FindLinkChecked(const FGuid& InSourceNodeGuid, const FGuid& InTargetNodeGuid)
+	{
+		const FSchematicGraphModel* ConstThis = this;
+		return const_cast<LinkType*>(ConstThis->FindLinkChecked<LinkType>(InSourceNodeGuid, InTargetNodeGuid));
+	}
+
 	bool IsLinkedTo(const FGuid& InSourceNodeGuid, const FGuid& InTargetNodeGuid) const;
 	TArray<const FSchematicGraphLink*> FindLinksOnNode(const FGuid& InNodeGuid) const;
 	TArray<const FSchematicGraphLink*> FindLinksOnSource(const FGuid& InSourceNodeGuid) const;
@@ -202,6 +250,10 @@ public:
 	ESchematicGraphVisibility::Type GetVisibilityForLink(const FGuid& InLinkGuid) const;
 	virtual ESchematicGraphVisibility::Type GetVisibilityForLink(const FSchematicGraphLink* InLink) const;
 
+	virtual bool IsAutoGroupingEnabled() const { return true; }
+	virtual float GetAutoGroupingDistance() const { return 32.f; }
+	virtual FSchematicGraphGroupNode* AddAutoGroupNode(); 
+
 	FOnSchematicGraphNodeAdded& OnNodeAdded() { return OnNodeAddedDelegate; }
 	FOnNodeRemoved& OnNodeRemoved() { return OnNodeRemovedDelegate; }
 	FOnSchematicGraphLinkAdded& OnLinkAdded() { return OnLinkAddedDelegate; }
@@ -212,6 +264,15 @@ public:
 
 	virtual bool GetForwardedNodeForDrag(FGuid& InOutGuid) const { return false; }
 
+	const FSchematicGraphGroupNode* GetLastExpandedNode() const;
+	virtual void SetLastExpandedNode(const FSchematicGraphGroupNode* InGroupNode);
+	void ClearLastExpandedNode()
+	{
+		SetLastExpandedNode(nullptr);
+	}
+
+	void Tick(float InDeltaTime);
+
 protected:
 
 	TArray<TSharedPtr<FSchematicGraphNode>> Nodes;
@@ -220,6 +281,7 @@ protected:
 	TMap<FGuid, TSharedPtr<FSchematicGraphLink>> LinkByGuid;
 	TMap<uint32, TSharedPtr<FSchematicGraphLink>> LinkByHash;
 	TMap<FGuid, TTuple<TArray<FGuid>, TArray<FGuid>>> NodeGuidToLinkGuids;
+	FGuid LastExpandedNode = FGuid();
 
 	FOnSchematicGraphNodeAdded OnNodeAddedDelegate;
 	FOnNodeRemoved OnNodeRemovedDelegate;
diff --git a/Engine/Source/Editor/AnimationEditorWidgets/Public/SchematicGraphPanel/SchematicGraphNode.h b/Engine/Source/Editor/AnimationEditorWidgets/Public/SchematicGraphPanel/SchematicGraphNode.h
--- a/Engine/Source/Editor/AnimationEditorWidgets/Public/SchematicGraphPanel/SchematicGraphNode.h
+++ b/Engine/Source/Editor/AnimationEditorWidgets/Public/SchematicGraphPanel/SchematicGraphNode.h
@@ -3,28 +3,50 @@
 
 #if WITH_EDITOR
 
+#include "Animation/AnimatedAttribute.h"
 #include "SchematicGraphTag.h"
+#include "Input/Events.h"
+#include "Input/Reply.h"
 
 #define SCHEMATICGRAPHNODE_BODY(ClassName, SuperClass) \
 SCHEMATICGRAPHELEMENT_BODY(ClassName, SuperClass, FSchematicGraphNode)
 
-class ANIMATIONEDITORWIDGETS_API FSchematicGraphNode
+class FSchematicGraphGroupNode;
+
+class ANIMATIONEDITORWIDGETS_API FSchematicGraphNode : public TSharedFromThis<FSchematicGraphNode>
 {
 public:
 
 	SCHEMATICGRAPHELEMENT_BODY_BASE(FSchematicGraphNode)
+
+	FSchematicGraphNode();
 	
 	const FGuid& GetGuid() const { return Guid; }
+	const FSchematicGraphModel* GetGraph() const { return Model; }
 	bool HasParentNode() const { return ParentNodeGuid.IsValid(); }
 	const FGuid& GetParentNodeGuid() const { return ParentNodeGuid; }
 	const FSchematicGraphNode* GetParentNode() const;
+	const FGuid& GetRootNodeGuid() const;
+	const FSchematicGraphNode* GetRootNode() const;
+	bool IsRootNode() const { return !HasParentNode(); }
+	const FSchematicGraphGroupNode* GetGroupNode() const;
 	int32 GetNumChildNodes() const { return ChildNodeGuids.Num(); }
 	const TArray<FGuid>& GetChildNodeGuids() const { return ChildNodeGuids; }
 	const FSchematicGraphNode* GetChildNode(int32 InChildNodeIndex) const;
-	virtual ESchematicGraphVisibility::Type GetVisibilityForChildNodes() const { return GetVisibility(); }
+	TOptional<ESchematicGraphVisibility::Type> GetVisibilityForChildNode(const FGuid& InChildGuid) const;
+	virtual TOptional<ESchematicGraphVisibility::Type> GetVisibilityForChildNode(const FSchematicGraphNode* InChildNode) const { return  TOptional<ESchematicGraphVisibility::Type>(); }
+	TOptional<FVector2d> GetPositionForChildNode(const FGuid& InChildGuid) const;
+	virtual TOptional<FVector2d> GetPositionForChildNode(const FSchematicGraphNode* InChildNode) const { return TOptional<FVector2d>(); }
+	TOptional<float> GetScaleForChildNode(const FGuid& InChildGuid) const;
+	virtual TOptional<float> GetScaleForChildNode(const FSchematicGraphNode* InChildNode) const { return TOptional<float>(); }
+	TOptional<bool> GetInteractivityForChildNode(const FGuid& InChildGuid) const;
+	virtual TOptional<bool> GetInteractivityForChildNode(const FSchematicGraphNode* InChildNode) const { return TOptional<bool>(); }
 	bool IsSelected() const { return bIsSelected; }
 	void SetSelected(bool bSelected = true) { bIsSelected = bSelected; }
-	virtual FVector2d GetPosition() const { return Position; }
+	virtual const FText& GetLabel() const { return Label; }
+	void SetLabel(const FText& InLabel) { Label = InLabel; }
+	virtual FVector2d GetPosition() const;
+	virtual void SetPosition(const FVector2d& InPosition) { Position = InPosition; }
 	virtual FVector2d GetPositionOffset() const { return PositionOffset; }
 	virtual void SetPositionOffset(const FVector2d& InPositionOffset) { PositionOffset = InPositionOffset; }
 	virtual float GetScaleOffset() const { return ScaleOffset; }
@@ -33,21 +55,28 @@ public:
 	virtual FLinearColor GetColor(int32 InLayerIndex) const { return Colors[InLayerIndex]; }
 	virtual const FSlateBrush* GetBrush(int32 InLayerIndex) const { return Brushes[InLayerIndex]; }
 	virtual bool IsAutoScaleEnabled() const { return false; }
-	virtual const FText& GetToolTip() const { return ToolTip; }
+	virtual const FText& GetToolTip() const;
 	virtual ESchematicGraphPlacementConstraint::Type GetPlacement() const { return Placement; }
 	virtual void SetPlacement(ESchematicGraphPlacementConstraint::Type InPlacement) { Placement = InPlacement; }
 	virtual ESchematicGraphVisibility::Type GetVisibility() const;
 	virtual void SetVisibility(ESchematicGraphVisibility::Type InVisibility) { Visibility = InVisibility; }
+	virtual bool IsInteractive() const;
 	virtual bool IsDragSupported() const { return bDragSupported; }
 	virtual void SetDragSupported(bool InDragSupported) { bDragSupported = InDragSupported;}
 
+	virtual FReply OnClicked(const FPointerEvent& InMouseEvent);
+	virtual void OnMouseEnter();
+	virtual void OnMouseLeave();
+	virtual void OnDragOver();
+	virtual void OnDragLeave();;
+
 	virtual FString GetDragDropDecoratorLabel() const;
 
 	template<typename TagType = FSchematicGraphTag>
 	TagType* AddTag(bool bNotify = true)
 	{
 		const TSharedPtr<FSchematicGraphTag> NewTag = MakeShareable(new TagType);
-		NewTag->Model = Model;
+		NewTag->Node = this;
 		Tags.Add(NewTag);
 		TagByGuid.Add(NewTag->GetGuid(), NewTag);
 
@@ -68,6 +97,13 @@ public:
 		return nullptr;
 	}
 
+	template<typename TagType = FSchematicGraphTag>
+	TagType* FindTag(const FGuid& InTagGuid)
+	{
+		const FSchematicGraphNode* ConstThis = this;
+		return const_cast<TagType*>(ConstThis->FindTag<TagType>(InTagGuid));
+	}
+
 	template<typename TagType = FSchematicGraphTag>
 	const TagType* FindTagChecked(const FGuid& InTagGuid) const
 	{
@@ -79,6 +115,13 @@ public:
 		return nullptr;
 	}
 
+	template<typename TagType = FSchematicGraphTag>
+	TagType* FindTagChecked(const FGuid& InTagGuid)
+	{
+		const FSchematicGraphNode* ConstThis = this;
+		return const_cast<TagType*>(ConstThis->FindTagChecked<TagType>(InTagGuid));
+	}
+
 	virtual bool RemoveTag(const FGuid& InTagGuid);
 
 	const TArray<TSharedPtr<FSchematicGraphTag>>& GetTags() const { return Tags; }
@@ -89,41 +132,86 @@ protected:
 
 	FSchematicGraphModel* Model = nullptr;
 	FGuid Guid = FGuid::NewGuid();
+	FText Label = FText();
 	FGuid ParentNodeGuid = FGuid();
 	TArray<FGuid> ChildNodeGuids; 
 	bool bIsSelected = false;
 	FVector2d Position = FVector2d::ZeroVector;
 	FVector2d PositionOffset = FVector2d::ZeroVector;
 	float ScaleOffset = 1.f;
-	TArray<FLinearColor> Colors = { FLinearColor::White };
-	TArray<const FSlateBrush*> Brushes = { nullptr };
+	TArray<FLinearColor> Colors;
+	TArray<const FSlateBrush*> Brushes;
 	FText ToolTip = FText();
 	ESchematicGraphPlacementConstraint::Type Placement = ESchematicGraphPlacementConstraint::Free;
 	ESchematicGraphVisibility::Type Visibility = ESchematicGraphVisibility::Visible;
 	bool bDragSupported = false;
 
+	static inline constexpr float ScaledUp = 1.25;
+	static inline constexpr float ScaledDown = 0.75;
+
 	TArray<TSharedPtr<FSchematicGraphTag>> Tags;
 	TMap<FGuid, TSharedPtr<FSchematicGraphTag>> TagByGuid;
 
 	friend class FSchematicGraphModel;
 };
 
-class FSchematicGraphGroupNode : public FSchematicGraphNode
+class ANIMATIONEDITORWIDGETS_API FSchematicGraphGroupNode : public FSchematicGraphNode
 {
 public:
 
 	SCHEMATICGRAPHNODE_BODY(FSchematicGraphGroupNode, FSchematicGraphNode)
 
+	FSchematicGraphGroupNode();
 	virtual ~FSchematicGraphGroupNode() override {}
-	bool IsExpanded() const { return bExpanded; }
-	void SetExpanded(bool InExpanded) { bExpanded = InExpanded; }
-	virtual ESchematicGraphVisibility::Type GetVisibilityForChildNodes() const override;
-	
+
+	// FSchematicGraphNode interface
+	virtual TOptional<ESchematicGraphVisibility::Type> GetVisibilityForChildNode(const FSchematicGraphNode* InChildNode) const override;
+	virtual TOptional<FVector2d> GetPositionForChildNode(const FSchematicGraphNode* InChildNode) const override;
+	virtual TOptional<float> GetScaleForChildNode(const FSchematicGraphNode* InChildNode) const override;
+	virtual TOptional<bool> GetInteractivityForChildNode(const FSchematicGraphNode* InChildNode) const override;
+	virtual FReply OnClicked(const FPointerEvent& InMouseEvent) override;
+	virtual void OnDragOver() override;
+	virtual void OnDragLeave() override;
+
+	bool IsExpanded() const;
+	bool IsExpanding() const;
+	bool IsCollapsing() const;
+	float GetExpansionState() const;
+	void SetExpanded(bool InExpanded, bool bAutoCloseParentGroups = false);
+	FLinearColor GetExpansionColor() const { return ExpansionColor; }
+	void SetExpansionColor(FLinearColor InExpansionColor) { ExpansionColor = InExpansionColor; }
+	float GetExpansionRadius() const { return ExpansionRadius; }
+	void SetExpansionRadius(float InExpansionRadius) { ExpansionRadius = InExpansionRadius; }
+	float GetDelayDuration(bool bEnter) const;
+	void SetDelayDuration(bool bEnter, float InDelayDuration) { (bEnter ? EnterDelayDuration : LeaveDelayDuration) = InDelayDuration; }
+
 protected:
 
-	bool bExpanded = false;
+	const TEasingAttributeInterpolator<float>::FSettings& GetAnimationSettings() const { return AnimationSettings; }
+	void SetAnimationSettings(const TEasingAttributeInterpolator<float>::FSettings& InSettings);
 
-	friend class FControlRigSchematicModel;
+	float ExpansionRadius = 60.f;
+	FLinearColor ExpansionColor = FLinearColor::White * FLinearColor(1.f, 1.f, 1.f, 0.5f);
+	float EnterDelayDuration = 0.4f;
+	float LeaveDelayDuration = 1.f;
+	TEasingAttributeInterpolator<float>::FSettings AnimationSettings;
+	TSharedPtr<TAnimatedAttribute<float>> ExpansionState;
+
+	friend class SSchematicGraphNode;
 };
 
+class ANIMATIONEDITORWIDGETS_API FSchematicGraphAutoGroupNode : public FSchematicGraphGroupNode
+{
+public:
+
+	SCHEMATICGRAPHNODE_BODY(FSchematicGraphAutoGroupNode, FSchematicGraphGroupNode)
+
+	virtual const FText& GetLabel() const override;
+
+protected:
+
+	mutable FText AutoGroupLabel;
+};
+
+
 #endif
\ No newline at end of file
diff --git a/Engine/Source/Editor/AnimationEditorWidgets/Public/SchematicGraphPanel/SchematicGraphStyle.h b/Engine/Source/Editor/AnimationEditorWidgets/Public/SchematicGraphPanel/SchematicGraphStyle.h
new file mode 100644
--- /dev/null
+++ b/Engine/Source/Editor/AnimationEditorWidgets/Public/SchematicGraphPanel/SchematicGraphStyle.h
@@ -0,0 +1,49 @@
+// Copyright Epic Games, Inc. All Rights Reserved.
+
+#pragma once
+
+#include "Styling/SlateStyle.h"
+#include "Styling/SlateStyleRegistry.h"
+#include "Styling/SlateStyleMacros.h"
+#include "Styling/StyleColors.h"
+#include "Brushes/SlateRoundedBoxBrush.h"
+
+class ANIMATIONEDITORWIDGETS_API FSchematicGraphStyle final
+	: public FSlateStyleSet
+{
+
+public:
+	FSchematicGraphStyle()
+		: FSlateStyleSet("SchematicGraphStyle")
+	{
+		const FVector2D Icon10x10(10.0f, 10.0f);
+		const FVector2D Icon14x14(14.0f, 14.0f);
+		const FVector2D Icon16x16(16.0f, 16.0f);
+		const FVector2D Icon20x20(20.0f, 20.0f);
+		const FVector2D Icon24x24(24.0f, 24.0f);
+		const FVector2D Icon32x32(32.0f, 32.0f);
+		const FVector2D Icon40x40(40.0f, 40.0f);
+		const FVector2D Icon128x128(128.0f, 128.0f);
+		const FString EngineEditorSlateDir = FPaths::EngineContentDir() / TEXT("Editor/Slate");
+		SetContentRoot(EngineEditorSlateDir);
+		SetCoreContentRoot(EngineEditorSlateDir);
+
+		Set("Schematic.Background", new IMAGE_BRUSH_SVG("SchematicGraph/Schematic_Background", Icon128x128));
+		Set("Schematic.Group", new IMAGE_BRUSH_SVG("SchematicGraph/Schematic_Group", Icon128x128));
+		Set("Schematic.Outline.Single", new IMAGE_BRUSH_SVG("SchematicGraph/Schematic_Outline_Single", Icon128x128));
+		Set("Schematic.Outline.Double", new IMAGE_BRUSH_SVG("SchematicGraph/Schematic_Outline_Double", Icon128x128));
+		Set("Schematic.Dot.Small", new IMAGE_BRUSH_SVG("SchematicGraph/Schematic_Dot_Small", Icon128x128));
+		Set("Schematic.Dot.Medium", new IMAGE_BRUSH_SVG("SchematicGraph/Schematic_Dot_Medium", Icon128x128));
+		Set("Schematic.Dot.Large", new IMAGE_BRUSH_SVG("SchematicGraph/Schematic_Dot_Large", Icon128x128));
+		Set("Schematic.Dot.Group", new IMAGE_BRUSH_SVG("SchematicGraph/Schematic_Dot_Group", Icon128x128));
+
+		FSlateStyleRegistry::RegisterSlateStyle(*this);
+	}
+
+	static FSchematicGraphStyle& Get();
+	
+	~FSchematicGraphStyle()
+	{
+		FSlateStyleRegistry::UnRegisterSlateStyle(*this);
+	}
+};
diff --git a/Engine/Source/Editor/AnimationEditorWidgets/Public/SchematicGraphPanel/SchematicGraphTag.h b/Engine/Source/Editor/AnimationEditorWidgets/Public/SchematicGraphPanel/SchematicGraphTag.h
--- a/Engine/Source/Editor/AnimationEditorWidgets/Public/SchematicGraphPanel/SchematicGraphTag.h
+++ b/Engine/Source/Editor/AnimationEditorWidgets/Public/SchematicGraphPanel/SchematicGraphTag.h
@@ -8,12 +8,14 @@
 #define SCHEMATICGRAPHTAG_BODY(ClassName, SuperClass) \
 SCHEMATICGRAPHELEMENT_BODY(ClassName, SuperClass, FSchematicGraphTag)
 
-class ANIMATIONEDITORWIDGETS_API FSchematicGraphTag
+class ANIMATIONEDITORWIDGETS_API FSchematicGraphTag : public TSharedFromThis<FSchematicGraphTag>
 {
 public:
 	
 	SCHEMATICGRAPHELEMENT_BODY_BASE(FSchematicGraphTag)
 
+	FSchematicGraphTag();
+
 	const FSchematicGraphNode* GetNode() const;
 	const FGuid& GetGuid() const { return Guid; }
 	virtual FLinearColor GetBackgroundColor() const { return BackgroundColor; }
@@ -26,9 +28,9 @@ public:
 	virtual void SetBackgroundBrush(const FSlateBrush* InBackgroundBrush) { BackgroundBrush = InBackgroundBrush; }
 	virtual const FSlateBrush* GetForegroundBrush() const { return ForegroundBrush; }
 	virtual void SetForegroundBrush(const FSlateBrush* InForegroundBrush) { ForegroundBrush = InForegroundBrush; }
-	virtual const FText& GetLabel() const { return Label; }
+	virtual FText GetLabel() const { return Label; }
 	virtual void SetLabel(const FText& InLabel) { Label = InLabel; }
-	virtual const FText& GetToolTip() const { return ToolTip; }
+	virtual FText GetToolTip() const { return ToolTip; }
 	virtual void SetToolTip(const FText& InToolTip) { ToolTip = InToolTip; }
 	virtual float GetPlacementAngle() const { return PlacementAngle; }
 	virtual void SetPlacementAngle(float InPlacementAngle) { PlacementAngle = InPlacementAngle; }
@@ -37,8 +39,7 @@ public:
 
 protected:
 
-	FSchematicGraphModel* Model = nullptr;
-	TWeakPtr<FSchematicGraphNode> Node;
+	FSchematicGraphNode* Node = nullptr;
 	FGuid Guid = FGuid::NewGuid();
 	FLinearColor BackgroundColor = FLinearColor(FColor(0, 112, 224));
 	FLinearColor ForegroundColor = FLinearColor::White;
@@ -54,4 +55,16 @@ protected:
 	friend class FSchematicGraphNode;
 };
 
+class ANIMATIONEDITORWIDGETS_API FSchematicGraphGroupTag : public FSchematicGraphTag
+{
+public:
+
+	SCHEMATICGRAPHTAG_BODY(FSchematicGraphGroupTag, FSchematicGraphTag)
+
+	virtual ~FSchematicGraphGroupTag() override {}
+
+	virtual ESchematicGraphVisibility::Type GetVisibility() const override;
+	virtual FText GetLabel() const override;
+};
+
 #endif
\ No newline at end of file
diff --git a/Engine/Source/Runtime/SlateCore/Public/Animation/AnimatedAttribute.h b/Engine/Source/Runtime/SlateCore/Public/Animation/AnimatedAttribute.h
--- a/Engine/Source/Runtime/SlateCore/Public/Animation/AnimatedAttribute.h
+++ b/Engine/Source/Runtime/SlateCore/Public/Animation/AnimatedAttribute.h
@@ -109,6 +109,10 @@ public:
 	void Set( const OtherType& InNewValue )
 	{
 		Getter.Unbind();
+		if(DesiredValue.IsSet() && !Interpolator->DesiredValue.IsSet())
+		{
+			Interpolator->SetValue(DesiredValue.GetValue());
+		}
 		DesiredValue = InNewValue;
 		Interpolator->SetValue(DesiredValue.GetValue());
 	}
@@ -124,10 +128,26 @@ public:
 	void Set( NumericType&& InNewValue )
 	{
 		Getter.Unbind();
+		if(DesiredValue.IsSet() && !Interpolator->DesiredValue.IsSet())
+		{
+			Interpolator->SetValue(DesiredValue.GetValue());
+		}
 		DesiredValue = InNewValue;
 		Interpolator->SetValue(DesiredValue.GetValue());
 	}
 
+	/**
+	 * Sets the attribute's value and disables animation. The interpolator
+	 * will be stopped and the value will become the new current value.
+	 * 
+	 * @param InNewValue  The value to set the attribute to
+     */
+	void SetValueAndStop( const NumericType& InNewValue )
+	{
+		Set(InNewValue);
+		Interpolator->Reset();
+	}
+	
 	/**
 	 * Sets the attribute's value and disables animation. The interpolator
 	 * will be stopped and the value will become the new current value.
@@ -196,6 +216,19 @@ public:
 		return IsSet() ? Get() : DefaultValue;
 	}
 
+	/**
+	 * Returns the desired value this attribute is trying to reach.
+	 */
+	const NumericType& GetDesiredValue() const
+	{
+		if(DesiredValue.IsSet())
+		{
+			return DesiredValue.GetValue();
+		}
+		static const NumericType EmptyResult = NumericType();
+		return EmptyResult;
+	}
+
 	/**
 	 * Returns the overall deltatime of the interpolator
 	 */
@@ -204,6 +237,14 @@ public:
 		return Interpolator->GetOverAllDeltaTime();
 	}
 
+	/**
+	 * Returns the currently set delay on the interpolator
+	 */
+	TOptional<NumericType> GetDelay() const
+	{
+		return Interpolator->Delay;
+	}
+
 	/**
 	 * Sets the attribute's delay as a one shot
 	 */
@@ -220,6 +261,14 @@ public:
 		Interpolator->SetTolerance(Tolerance);
 	}
 
+	/**
+	 * Returns true if the attribute is currently animating
+	 */
+	bool IsPlaying() const
+	{
+		return Interpolator->IsPlaying();
+	}
+
 	/**
 	 * Enables (or disables) the interpolator and returns values in interpolated or immediate mode
 	 */
diff --git a/Engine/Source/Runtime/SlateCore/Public/Animation/AttributeInterpolator.h b/Engine/Source/Runtime/SlateCore/Public/Animation/AttributeInterpolator.h
--- a/Engine/Source/Runtime/SlateCore/Public/Animation/AttributeInterpolator.h
+++ b/Engine/Source/Runtime/SlateCore/Public/Animation/AttributeInterpolator.h
@@ -85,6 +85,7 @@ public:
 		InterpolatedValue.Reset();
 		LastValue.Reset();
 		DesiredValue.Reset();
+		Delay.Reset();
 		if(bFireStopEvent)
 		{
 			InterpolationStopped.Broadcast(TOptional<NumericType>());
@@ -196,6 +197,18 @@ public:
 			}
 		}
 
+		if(Delay.IsSet())
+		{
+			const float DelayLeft = Delay.GetValue() - InDeltaTime;
+			if(DelayLeft > SMALL_NUMBER)
+			{
+				Delay = DelayLeft;
+				return;
+			}
+			Delay.Reset();
+			InDeltaTime = -DelayLeft;
+		}
+
 		(void)PlayIfStopped();
 		Interpolate();
 		
