Schematic: Disable overlay based on bounds

#rnx

diff --git a/Engine/Source/Editor/AnimationEditorWidgets/Private/SchematicGraphPanel/SSchematicGraphPanel.cpp b/Engine/Source/Editor/AnimationEditorWidgets/Private/SchematicGraphPanel/SSchematicGraphPanel.cpp
--- a/Engine/Source/Editor/AnimationEditorWidgets/Private/SchematicGraphPanel/SSchematicGraphPanel.cpp
+++ b/Engine/Source/Editor/AnimationEditorWidgets/Private/SchematicGraphPanel/SSchematicGraphPanel.cpp
@@ -808,6 +808,7 @@ int32 SSchematicGraphPanel::OnPaint(const FPaintArgs& Args, const FGeometry& All
 		}
 	}
 
+	/*
 	FBox2d Bounds(EForceInit::ForceInit);
 	for (int32 ChildIndex = 0; ChildIndex < ArrangedChildren.Num(); ++ChildIndex)
 	{
@@ -822,21 +823,25 @@ int32 SSchematicGraphPanel::OnPaint(const FPaintArgs& Args, const FGeometry& All
 		Bounds += NodeCenterByIndex[ChildIndex] + FVector2d(120, 120);
 		Bounds += NodeCenterByIndex[ChildIndex] - FVector2d(120, 120);
 	}
+	*/
 
 	const float Alpha = FadeBackgroundAlpha->Get();
-	if (Alpha > 0.f && !Bounds.GetSize().IsNearlyZero())
+	if (Alpha > 0.f) // && !Bounds.GetSize().IsNearlyZero())
 	{
 		const FSlateBrush* Brush = FAppStyle::GetBrush(TEXT("Graph.Panel.SolidBackground"));
-		const FLinearColor TransparentGrey = FLinearColor(0.5, 0.5, 0.5, Alpha);
+		const FLinearColor TransparentGrey = FLinearColor(0.0, 0.0, 0.0, Alpha);
 		FSlateDrawElement::MakeBox(
 				OutDrawElements,
 				BackgroundLayer,
-				AllottedGeometry.ToPaintGeometry(Bounds.GetSize(), FSlateLayoutTransform(Bounds.Min)),
+				AllottedGeometry.ToPaintGeometry(),
+				//AllottedGeometry.ToPaintGeometry(Bounds.GetSize(), FSlateLayoutTransform(Bounds.Min)),
 				Brush,
 				ESlateDrawEffect::None,
 				TransparentGrey
 			);
 
+		/*
+		// Render an outline around the bounds
 		TArray<FVector2d> LinePoints;
 		LinePoints.AddZeroed(2);
 
@@ -893,6 +898,7 @@ int32 SSchematicGraphPanel::OnPaint(const FPaintArgs& Args, const FGeometry& All
 				FLinearColor(0.7, 0.7, 0.7, 0.95),
 				true);
 		}
+		*/
 	}
 	
 	// draw all of the links
