SchematicGraph: Render label with backdrop

#rnx

diff --git a/Engine/Source/Editor/AnimationEditorWidgets/Private/SchematicGraphPanel/SSchematicGraphPanel.cpp b/Engine/Source/Editor/AnimationEditorWidgets/Private/SchematicGraphPanel/SSchematicGraphPanel.cpp
--- a/Engine/Source/Editor/AnimationEditorWidgets/Private/SchematicGraphPanel/SSchematicGraphPanel.cpp
+++ b/Engine/Source/Editor/AnimationEditorWidgets/Private/SchematicGraphPanel/SSchematicGraphPanel.cpp
@@ -238,7 +238,6 @@ int32 SSchematicGraphNode::OnPaint(const FPaintArgs& Args, const FGeometry& Allo
 			if(!LabelSize.IsNearlyZero())
 			{
 				NewLayerId++;
-				
 				const FVector2d LabelOffset = TagCenter - LabelSize * 0.5;
 
 				FSlateDrawElement::MakeText(
@@ -283,8 +282,22 @@ int32 SSchematicGraphNode::OnPaint(const FPaintArgs& Args, const FGeometry& Allo
 		const FVector2d NodeBottomCenter = AllottedGeometry.GetLocalSize() * FVector2d(0.5f, 1.f) + FVector2d(0.f, 8.f);
 		const FVector2d NodeLabelOffset = NodeBottomCenter - NodeLabelSize * FVector2d(0.5, 0);
 
+		static const FVector2d LabelBackgroundPadding = FVector2d(2.f);
+		const FVector2d LabelBackgroundSize = NodeLabelSize + LabelBackgroundPadding * 2.f;
+		const FVector2d LabelBackgroundOffset = NodeLabelOffset - LabelBackgroundPadding;
+		static const FSlateBrush* LabelBackgroundBrush = FSchematicGraphStyle::Get().GetBrush( "Schematic.Label.Background");
+
 		NewLayerId++;
+		FSlateDrawElement::MakeBox(
+			OutDrawElements,
+			NewLayerId,
+			AllottedGeometry.ToPaintGeometry(LabelBackgroundSize, FSlateLayoutTransform(LabelBackgroundOffset)),
+			LabelBackgroundBrush,
+			ESlateDrawEffect::None,
+			FLinearColor(0.2f, 0.2f, 0.2f, 0.8f) * FadedOutFactor
+		);
 
+		NewLayerId++;
 		FSlateDrawElement::MakeText(
 			OutDrawElements,
 			NewLayerId,
diff --git a/Engine/Source/Editor/AnimationEditorWidgets/Public/SchematicGraphPanel/SchematicGraphStyle.h b/Engine/Source/Editor/AnimationEditorWidgets/Public/SchematicGraphPanel/SchematicGraphStyle.h
--- a/Engine/Source/Editor/AnimationEditorWidgets/Public/SchematicGraphPanel/SchematicGraphStyle.h
+++ b/Engine/Source/Editor/AnimationEditorWidgets/Public/SchematicGraphPanel/SchematicGraphStyle.h
@@ -37,6 +37,8 @@ public:
 		Set("Schematic.Dot.Large", new IMAGE_BRUSH_SVG("SchematicGraph/Schematic_Dot_Large", Icon128x128));
 		Set("Schematic.Dot.Group", new IMAGE_BRUSH_SVG("SchematicGraph/Schematic_Dot_Group", Icon128x128));
 
+		Set("Schematic.Label.Background", new FSlateRoundedBoxBrush(FStyleColors::White, 4.0f, FStyleColors::Transparent, 0.0f));
+
 		FSlateStyleRegistry::RegisterSlateStyle(*this);
 	}
 
