SchematicView: remove notion of placement constraint

#rnx

diff --git a/Engine/Plugins/Animation/ControlRig/Source/ControlRigEditor/Private/Editor/ControlRigSchematicModel.cpp b/Engine/Plugins/Animation/ControlRig/Source/ControlRigEditor/Private/Editor/ControlRigSchematicModel.cpp
--- a/Engine/Plugins/Animation/ControlRig/Source/ControlRigEditor/Private/Editor/ControlRigSchematicModel.cpp
+++ b/Engine/Plugins/Animation/ControlRig/Source/ControlRigEditor/Private/Editor/ControlRigSchematicModel.cpp
@@ -913,21 +913,6 @@ FLinearColor FControlRigSchematicModel::GetColorForNode(const FSchematicGraphNod
 	return Super::GetColorForNode(InNode, InLayerIndex);
 }
 
-ESchematicGraphPlacementConstraint::Type FControlRigSchematicModel::GetPlacementForNode(const FSchematicGraphNode* InNode) const
-{
-	if(const FControlRigSchematicRigElementKeyNode* Node = Cast<FControlRigSchematicRigElementKeyNode>(InNode))
-	{
-		if(Node->GetKey().Type == ERigElementType::Connector)
-		{
-			if(!IsConnectorResolved(Node->GetKey()))
-			{
-				return ESchematicGraphPlacementConstraint::BottomRight;
-			}
-		}
-	}
-	return Super::GetPlacementForNode(InNode);
-}
-
 ESchematicGraphVisibility::Type FControlRigSchematicModel::GetVisibilityForNode(const FSchematicGraphNode* InNode) const
 {
 	if(ControlRigBlueprint.IsValid())
@@ -936,48 +921,26 @@ ESchematicGraphVisibility::Type FControlRigSchematicModel::GetVisibilityForNode(
 		{
 			if(Node->GetKey().Type == ERigElementType::Connector)
 			{
-				// hide secondary connectors if the primary connector isn't resolved yet
 				if (const UModularRig* ControlRig = Cast<UModularRig>(ControlRigBlueprint->GetDebuggedControlRig()))
 				{
 					if(const URigHierarchy* Hierarchy = ControlRig->GetHierarchy())
 					{
 						if(const FRigConnectorElement* Connector = Hierarchy->Find<FRigConnectorElement>(Node->GetKey()))
 						{
-							if(Connector->IsPrimary())
+							if(IsConnectorResolved(Connector->GetKey()))
 							{
-								// if the module is a root module
-								if(const FRigModuleInstance* ModuleInstance = ControlRig->FindModule(Connector))
-								{
-									if(ModuleInstance->IsRootModule())
-									{
-										if(IsConnectorResolved(Connector->GetKey()))
-										{
-											return ESchematicGraphVisibility::Hidden;
-										}
-									}
-								}
-							}
-							else
-							{
-								if(Connector->Settings.bOptional)
-								{
-									if(!IsConnectorResolved(Connector->GetKey()))
-									{
-										return ESchematicGraphVisibility::Hidden;
-									}
-								}
-								
-								if(const FRigModuleInstance* ModuleInstance = ControlRig->FindModule(Connector))
+								if(Connector->IsPrimary())
 								{
-									if(const FRigConnectorElement* PrimaryConnector = ModuleInstance->FindPrimaryConnector())
+									if(const FRigModuleInstance* ModuleInstance = ControlRig->FindModule(Connector))
 									{
-										if(!IsConnectorResolved(PrimaryConnector->GetKey()))
+										if(ModuleInstance->IsRootModule())
 										{
 											return ESchematicGraphVisibility::Hidden;
 										}
 									}
 								}
 							}
+							return ESchematicGraphVisibility::Hidden;
 						}
 					}
 				}
diff --git a/Engine/Plugins/Animation/ControlRig/Source/ControlRigEditor/Private/Editor/ControlRigSchematicModel.h b/Engine/Plugins/Animation/ControlRig/Source/ControlRigEditor/Private/Editor/ControlRigSchematicModel.h
--- a/Engine/Plugins/Animation/ControlRig/Source/ControlRigEditor/Private/Editor/ControlRigSchematicModel.h
+++ b/Engine/Plugins/Animation/ControlRig/Source/ControlRigEditor/Private/Editor/ControlRigSchematicModel.h
@@ -100,7 +100,6 @@ public:
 	const FSlateBrush* GetBrushForKey(const FRigElementKey& InKey, const FSchematicGraphNode* InNode) const;
 	virtual const FSlateBrush* GetBrushForNode(const FSchematicGraphNode* InNode, int32 InLayerIndex) const override;
 	virtual FLinearColor GetColorForNode(const FSchematicGraphNode* InNode, int32 InLayerIndex) const override;
-	virtual ESchematicGraphPlacementConstraint::Type GetPlacementForNode(const FSchematicGraphNode* InNode) const override;
 	virtual ESchematicGraphVisibility::Type GetVisibilityForNode(const FSchematicGraphNode* InNode) const override;
 	virtual const FSlateBrush* GetBrushForLink(const FSchematicGraphLink* InLink) const override;
 	virtual FLinearColor GetColorForLink(const FSchematicGraphLink* InLink) const override;
diff --git a/Engine/Source/Editor/AnimationEditorWidgets/Private/SchematicGraphPanel/SSchematicGraphPanel.cpp b/Engine/Source/Editor/AnimationEditorWidgets/Private/SchematicGraphPanel/SSchematicGraphPanel.cpp
--- a/Engine/Source/Editor/AnimationEditorWidgets/Private/SchematicGraphPanel/SSchematicGraphPanel.cpp
+++ b/Engine/Source/Editor/AnimationEditorWidgets/Private/SchematicGraphPanel/SSchematicGraphPanel.cpp
@@ -1019,13 +1019,6 @@ void SSchematicGraphPanel::Tick(float DeltaTime)
 		DragDropOpFromOutside.Reset();
 	}
 
-	NodesTopLeft.Reset();
-	NodesTopRight.Reset();
-	NodesBottomLeft.Reset();
-	NodesBottomRight.Reset();
-
-	const FVector2D Size = GetCachedGeometry().Size;
-
 	for (int32 i=0; i<Children.Num(); ++i)
 	{
 		TSharedRef<SSchematicGraphNode> Widget = GetChild(i);
@@ -1042,35 +1035,6 @@ void SSchematicGraphPanel::Tick(float DeltaTime)
 			continue;
 		}
 
-		// collect the nodes displayed in the corners
-		switch(GraphData->GetPlacementForNode(Widget->GetNodeData()))
-		{
-			case ESchematicGraphPlacementConstraint::Free:
-			{
-				break;
-			}
-			case ESchematicGraphPlacementConstraint::TopLeft:
-			{
-				NodesTopLeft.Add(Widget->GetGuid());
-				break;
-			}
-			case ESchematicGraphPlacementConstraint::TopRight:
-			{
-				NodesTopRight.Add(Widget->GetGuid());
-				break;
-			}
-			case ESchematicGraphPlacementConstraint::BottomLeft:
-			{
-				NodesBottomLeft.Add(Widget->GetGuid());
-				break;
-			}
-			case ESchematicGraphPlacementConstraint::BottomRight:
-			{
-				NodesBottomRight.Add(Widget->GetGuid());
-				break;
-			}
-		}
-
 		if (Widget->IsBeingDragged())
 		{
 			if (bIsDragDropping)
@@ -1182,55 +1146,7 @@ FVector2d SSchematicGraphPanel::GetPositionForNode(FGuid InNodeGuid) const
 		if(const FSchematicGraphNode* Node = GraphData->FindNode(InNodeGuid))
 		{
 			FVector2d Position = GraphData->GetPositionOffsetForNode(Node);
-
-			const FVector2d OriginalSize = NodeWidget->GetOriginalSize();
-			const FVector2D Size = GetCachedGeometry().Size;
-			const FVector2d HalfOriginalSize = OriginalSize * 0.5;
-
-			switch(GraphData->GetPlacementForNode(Node))
-			{
-				case ESchematicGraphPlacementConstraint::Free:
-				{
-					Position += GraphData->GetPositionForNode(Node);
-					break;
-				}
-				case ESchematicGraphPlacementConstraint::TopLeft:
-				{
-					const int32 NodeIndexInCorner = NodesTopLeft.IndexOfByKey(InNodeGuid);
-					if(NodeIndexInCorner != INDEX_NONE)
-					{
-						Position += FVector2d(PaddingLeft + HalfOriginalSize.X, PaddingTop + HalfOriginalSize.Y + ((OriginalSize.Y + PaddingInterNode) * NodeIndexInCorner));
-					}
-					break;
-				}
-				case ESchematicGraphPlacementConstraint::TopRight:
-				{
-					const int32 NodeIndexInCorner = NodesTopRight.IndexOfByKey(InNodeGuid);
-					if(NodeIndexInCorner != INDEX_NONE)
-					{
-						Position += FVector2d(Size.X - PaddingRight - HalfOriginalSize.X, PaddingTop + HalfOriginalSize.Y + ((OriginalSize.Y + PaddingInterNode) * NodeIndexInCorner));
-					}
-					break;
-				}
-				case ESchematicGraphPlacementConstraint::BottomLeft:
-				{
-					const int32 NodeIndexInCorner = NodesBottomLeft.IndexOfByKey(InNodeGuid);
-					if(NodeIndexInCorner != INDEX_NONE)
-					{
-						Position += FVector2d(PaddingLeft + HalfOriginalSize.X, Size.Y - PaddingBottom - HalfOriginalSize.Y - ((OriginalSize.Y + PaddingInterNode)  * NodeIndexInCorner));
-					}
-					break;
-				}
-				case ESchematicGraphPlacementConstraint::BottomRight:
-				{
-					const int32 NodeIndexInCorner = NodesBottomRight.IndexOfByKey(InNodeGuid);
-					if(NodeIndexInCorner != INDEX_NONE)
-					{
-						Position += FVector2d(Size.X - PaddingRight - HalfOriginalSize.X, Size.Y - PaddingBottom - HalfOriginalSize.Y - ((OriginalSize.Y + PaddingInterNode)  * NodeIndexInCorner));
-					}
-					break;
-				}
-			}
+			Position += GraphData->GetPositionForNode(Node);
 			return Position;
 		}
 	}
@@ -1347,7 +1263,6 @@ void SSchematicGraphPanel::UpdatePerNodeCaches(bool bRemoveNodesFromAutoGroups)
 			}
 			Cache.bHasParent = Node->HasParentNode();
 			Cache.Visibility = GraphData->GetVisibilityForNode(Node);
-			Cache.Placement = GraphData->GetPlacementForNode(Node);
 			Cache.bIsAutoScaling = !Cache.bHasParent && !Widget->bIsBeingDragged && Widget->EnableAutoScale.Get();
 			Cache.Position = GetPositionForNode(Node->GetGuid());
 			const FVector2d NodeSize = GraphData->GetSizeForNode(Node->GetGuid());
@@ -1394,10 +1309,6 @@ void SSchematicGraphPanel::UpdateAutoGroupingForNodes()
 		{
 			continue;
 		}
-		if(PerNodeCaches[i].Placement != ESchematicGraphPlacementConstraint::Free)
-		{
-			continue;
-		}
 
 		TSharedRef<SSchematicGraphNode> Widget = GetChild(i);
 		const FSchematicGraphNode* Node = Widget->GetNodeData();
diff --git a/Engine/Source/Editor/AnimationEditorWidgets/Private/SchematicGraphPanel/SchematicGraphModel.cpp b/Engine/Source/Editor/AnimationEditorWidgets/Private/SchematicGraphPanel/SchematicGraphModel.cpp
--- a/Engine/Source/Editor/AnimationEditorWidgets/Private/SchematicGraphPanel/SchematicGraphModel.cpp
+++ b/Engine/Source/Editor/AnimationEditorWidgets/Private/SchematicGraphPanel/SchematicGraphModel.cpp
@@ -333,21 +333,6 @@ FText FSchematicGraphModel::GetToolTipForNode(const FSchematicGraphNode* InNode)
 	return InNode->GetToolTip();
 }
 
-ESchematicGraphPlacementConstraint::Type FSchematicGraphModel::GetPlacementForNode(const FGuid& InNodeGuid) const
-{
-	if(const FSchematicGraphNode* Node = FindNode(InNodeGuid))
-	{
-		return GetPlacementForNode(Node);
-	}
-	return ESchematicGraphPlacementConstraint::Free;
-}
-
-ESchematicGraphPlacementConstraint::Type FSchematicGraphModel::GetPlacementForNode(const FSchematicGraphNode* InNode) const
-{
-	check(InNode);
-	return InNode->GetPlacement();
-}
-
 ESchematicGraphVisibility::Type FSchematicGraphModel::GetVisibilityForNode(const FGuid& InNodeGuid) const
 {
 	if(const FSchematicGraphNode* Node = FindNode(InNodeGuid))
diff --git a/Engine/Source/Editor/AnimationEditorWidgets/Public/SchematicGraphPanel/SSchematicGraphPanel.h b/Engine/Source/Editor/AnimationEditorWidgets/Public/SchematicGraphPanel/SSchematicGraphPanel.h
--- a/Engine/Source/Editor/AnimationEditorWidgets/Public/SchematicGraphPanel/SSchematicGraphPanel.h
+++ b/Engine/Source/Editor/AnimationEditorWidgets/Public/SchematicGraphPanel/SSchematicGraphPanel.h
@@ -262,10 +262,6 @@ private:
 	FOnLeaveDrag OnLeaveDragDelegate;
 	FOnCancelDrag OnCancelDragDelegate;
 	FOnDrop OnDropDelegate;
-	TArray<FGuid> NodesTopLeft;
-	TArray<FGuid> NodesTopRight;
-	TArray<FGuid> NodesBottomLeft;
-	TArray<FGuid> NodesBottomRight;
 	TMap<FGuid, TSharedPtr<SSchematicGraphNode>> NodeByGuid;
 
 	struct FPerNodeCache
@@ -273,7 +269,6 @@ private:
 		FPerNodeCache()
 			: bHasParent(false)
 			, Visibility(ESchematicGraphVisibility::Visible)
-			, Placement(ESchematicGraphPlacementConstraint::Free)
 			, bIsAutoScaling(false)
 			, Position(FVector2d::ZeroVector)
 			, Radius(0.0)
@@ -281,7 +276,6 @@ private:
 
 		bool bHasParent;
 		ESchematicGraphVisibility::Type Visibility;
-		ESchematicGraphPlacementConstraint::Type Placement;
 		bool bIsAutoScaling;
 		FVector2d Position;
 		double Radius;
diff --git a/Engine/Source/Editor/AnimationEditorWidgets/Public/SchematicGraphPanel/SchematicGraphDefines.h b/Engine/Source/Editor/AnimationEditorWidgets/Public/SchematicGraphPanel/SchematicGraphDefines.h
--- a/Engine/Source/Editor/AnimationEditorWidgets/Public/SchematicGraphPanel/SchematicGraphDefines.h
+++ b/Engine/Source/Editor/AnimationEditorWidgets/Public/SchematicGraphPanel/SchematicGraphDefines.h
@@ -13,18 +13,6 @@ class FSchematicGraphNode;
 class FSchematicGraphLink;
 class FSchematicGraphModel;
 
-namespace ESchematicGraphPlacementConstraint
-{
-	enum Type : int
-	{
-		Free,
-		TopLeft,
-		TopRight,
-		BottomLeft,
-		BottomRight
-	};
-}
-
 namespace ESchematicGraphVisibility
 {
 	enum Type : int
diff --git a/Engine/Source/Editor/AnimationEditorWidgets/Public/SchematicGraphPanel/SchematicGraphModel.h b/Engine/Source/Editor/AnimationEditorWidgets/Public/SchematicGraphPanel/SchematicGraphModel.h
--- a/Engine/Source/Editor/AnimationEditorWidgets/Public/SchematicGraphPanel/SchematicGraphModel.h
+++ b/Engine/Source/Editor/AnimationEditorWidgets/Public/SchematicGraphPanel/SchematicGraphModel.h
@@ -101,8 +101,6 @@ public:
 	virtual const FSlateBrush* GetBrushForNode(const FSchematicGraphNode* InNode, int32 InLayerIndex) const;
 	FText GetToolTipForNode(const FGuid& InNodeGuid) const;
 	virtual FText GetToolTipForNode(const FSchematicGraphNode* InNode) const;
-	ESchematicGraphPlacementConstraint::Type GetPlacementForNode(const FGuid& InNodeGuid) const;
-	virtual ESchematicGraphPlacementConstraint::Type GetPlacementForNode(const FSchematicGraphNode* InNode) const;
 	ESchematicGraphVisibility::Type GetVisibilityForNode(const FGuid& InNodeGuid) const;
 	virtual ESchematicGraphVisibility::Type GetVisibilityForNode(const FSchematicGraphNode* InNode) const;
 	TOptional<ESchematicGraphVisibility::Type> GetVisibilityForChildNode(const FGuid& InParentNodeGuid, const FGuid& InChildNodeGuid) const;
diff --git a/Engine/Source/Editor/AnimationEditorWidgets/Public/SchematicGraphPanel/SchematicGraphNode.h b/Engine/Source/Editor/AnimationEditorWidgets/Public/SchematicGraphPanel/SchematicGraphNode.h
--- a/Engine/Source/Editor/AnimationEditorWidgets/Public/SchematicGraphPanel/SchematicGraphNode.h
+++ b/Engine/Source/Editor/AnimationEditorWidgets/Public/SchematicGraphPanel/SchematicGraphNode.h
@@ -56,8 +56,6 @@ public:
 	virtual const FSlateBrush* GetBrush(int32 InLayerIndex) const { return Brushes[InLayerIndex]; }
 	virtual bool IsAutoScaleEnabled() const { return false; }
 	virtual const FText& GetToolTip() const;
-	virtual ESchematicGraphPlacementConstraint::Type GetPlacement() const { return Placement; }
-	virtual void SetPlacement(ESchematicGraphPlacementConstraint::Type InPlacement) { Placement = InPlacement; }
 	virtual ESchematicGraphVisibility::Type GetVisibility() const;
 	virtual void SetVisibility(ESchematicGraphVisibility::Type InVisibility) { Visibility = InVisibility; }
 	virtual bool IsInteractive() const;
@@ -152,7 +150,6 @@ protected:
 	TArray<FLinearColor> Colors;
 	TArray<const FSlateBrush*> Brushes;
 	FText ToolTip = FText();
-	ESchematicGraphPlacementConstraint::Type Placement = ESchematicGraphPlacementConstraint::Free;
 	ESchematicGraphVisibility::Type Visibility = ESchematicGraphVisibility::Visible;
 	bool bDragSupported = false;
 
