Schematic: Update label background color

#rnx

diff --git a/Engine/Source/Editor/AnimationEditorWidgets/Private/SchematicGraphPanel/SSchematicGraphPanel.cpp b/Engine/Source/Editor/AnimationEditorWidgets/Private/SchematicGraphPanel/SSchematicGraphPanel.cpp
--- a/Engine/Source/Editor/AnimationEditorWidgets/Private/SchematicGraphPanel/SSchematicGraphPanel.cpp
+++ b/Engine/Source/Editor/AnimationEditorWidgets/Private/SchematicGraphPanel/SSchematicGraphPanel.cpp
@@ -287,6 +287,9 @@ int32 SSchematicGraphNode::OnPaint(const FPaintArgs& Args, const FGeometry& Allo
 		const FVector2d LabelBackgroundOffset = NodeLabelOffset - LabelBackgroundPadding;
 		static const FSlateBrush* LabelBackgroundBrush = FSchematicGraphStyle::Get().GetBrush( "Schematic.Label.Background");
 
+		static const FColor LabelBackgroundColorHex = FColor::FromHex(TEXT("#0F0F0F"));
+		static const FLinearColor LabelBackgroundColor = FLinearColor(LabelBackgroundColorHex) * FLinearColor(1.f, 1.f, 1.f, 0.7f); 
+		
 		NewLayerId++;
 		FSlateDrawElement::MakeBox(
 			OutDrawElements,
@@ -294,7 +297,7 @@ int32 SSchematicGraphNode::OnPaint(const FPaintArgs& Args, const FGeometry& Allo
 			AllottedGeometry.ToPaintGeometry(LabelBackgroundSize, FSlateLayoutTransform(LabelBackgroundOffset)),
 			LabelBackgroundBrush,
 			ESlateDrawEffect::None,
-			FLinearColor(0.2f, 0.2f, 0.2f, 0.8f) * FadedOutFactor
+			LabelBackgroundColor * FadedOutFactor
 		);
 
 		NewLayerId++;
