Schematic Graph: Use separate SVG for tag background

#rnx

diff --git a/Engine/Source/Editor/AnimationEditorWidgets/Private/SchematicGraphPanel/SchematicGraphTag.cpp b/Engine/Source/Editor/AnimationEditorWidgets/Private/SchematicGraphPanel/SchematicGraphTag.cpp
--- a/Engine/Source/Editor/AnimationEditorWidgets/Private/SchematicGraphPanel/SchematicGraphTag.cpp
+++ b/Engine/Source/Editor/AnimationEditorWidgets/Private/SchematicGraphPanel/SchematicGraphTag.cpp
@@ -10,7 +10,7 @@
 
 FSchematicGraphTag::FSchematicGraphTag()
 {
-	static const FSlateBrush* CircleBrush = FSchematicGraphStyle::Get().GetBrush( "Schematic.Background");
+	static const FSlateBrush* CircleBrush = FSchematicGraphStyle::Get().GetBrush( "Schematic.Tag.Background");
 	BackgroundBrush = CircleBrush;
 }
 
diff --git a/Engine/Source/Editor/AnimationEditorWidgets/Public/SchematicGraphPanel/SchematicGraphStyle.h b/Engine/Source/Editor/AnimationEditorWidgets/Public/SchematicGraphPanel/SchematicGraphStyle.h
--- a/Engine/Source/Editor/AnimationEditorWidgets/Public/SchematicGraphPanel/SchematicGraphStyle.h
+++ b/Engine/Source/Editor/AnimationEditorWidgets/Public/SchematicGraphPanel/SchematicGraphStyle.h
@@ -36,6 +36,7 @@ public:
 		Set("Schematic.Dot.Medium", new IMAGE_BRUSH_SVG("SchematicGraph/Schematic_Dot_Medium", Icon128x128));
 		Set("Schematic.Dot.Large", new IMAGE_BRUSH_SVG("SchematicGraph/Schematic_Dot_Large", Icon128x128));
 		Set("Schematic.Dot.Group", new IMAGE_BRUSH_SVG("SchematicGraph/Schematic_Dot_Group", Icon128x128));
+		Set("Schematic.Tag.Background", new IMAGE_BRUSH_SVG("SchematicGraph/Schematic_Tag_Background", Icon128x128));
 
 		Set("Schematic.Label.Background", new FSlateRoundedBoxBrush(FStyleColors::White, 4.0f, FStyleColors::Transparent, 0.0f));
 
