Schematic: Node labels not drawing when not in full screen

#preflighit skip
#rnx

diff --git a/Engine/Source/Editor/AnimationEditorWidgets/Private/SchematicGraphPanel/SSchematicGraphPanel.cpp b/Engine/Source/Editor/AnimationEditorWidgets/Private/SchematicGraphPanel/SSchematicGraphPanel.cpp
--- a/Engine/Source/Editor/AnimationEditorWidgets/Private/SchematicGraphPanel/SSchematicGraphPanel.cpp
+++ b/Engine/Source/Editor/AnimationEditorWidgets/Private/SchematicGraphPanel/SSchematicGraphPanel.cpp
@@ -257,8 +257,8 @@ int32 SSchematicGraphNode::OnPaint(const FPaintArgs& Args, const FGeometry& Allo
 	bool bDrawLabel = true;
 	if(bDrawLabel)
 	{
-		const FVector2D MouseCursorLocation = FSlateApplication::Get().GetCursorPos();
-		if((AllottedGeometry.GetAbsolutePosition() - MouseCursorLocation).Length() > AllottedGeometry.GetAbsoluteSize().GetMax() * 0.5f + 8.f)
+		const FVector2D MouseCursorLocation = FSlateApplication::Get().GetCursorPos() - Args.GetWindowToDesktopTransform();
+		if((AllottedGeometry.GetAbsolutePositionAtCoordinates(FVector2d(0.5f, 0.5f)) - MouseCursorLocation).Length() > AllottedGeometry.GetAbsoluteSize().GetMax() * 0.5f + 8.f)
 		{
 			bDrawLabel = false;
 		}
