Schematic: make sure to render group background behind nodes

#rnx

diff --git a/Engine/Source/Editor/AnimationEditorWidgets/Private/SchematicGraphPanel/SSchematicGraphPanel.cpp b/Engine/Source/Editor/AnimationEditorWidgets/Private/SchematicGraphPanel/SSchematicGraphPanel.cpp
--- a/Engine/Source/Editor/AnimationEditorWidgets/Private/SchematicGraphPanel/SSchematicGraphPanel.cpp
+++ b/Engine/Source/Editor/AnimationEditorWidgets/Private/SchematicGraphPanel/SSchematicGraphPanel.cpp
@@ -114,7 +114,7 @@ FVector2D SSchematicGraphNode::ComputeDesiredSize(float LayoutScaleMultiplier) c
 
 int32 SSchematicGraphNode::OnPaint(const FPaintArgs& Args, const FGeometry& AllottedGeometry, const FSlateRect& MyCullingRect, FSlateWindowElementList& OutDrawElements, int32 LayerId, const FWidgetStyle& InWidgetStyle, bool bParentEnabled) const
 {
-	int32 NewLayerId = SNodePanel::SNode::OnPaint(Args, AllottedGeometry, MyCullingRect, OutDrawElements, LayerId, InWidgetStyle, bParentEnabled);
+	const int32 NewLayerId = SNodePanel::SNode::OnPaint(Args, AllottedGeometry, MyCullingRect, OutDrawElements, LayerId, InWidgetStyle, bParentEnabled);
 
 	const FVector2d CurSize = Size->Get() * Scale->Get();
 	const FVector2d SizeOffset = (CurSize-OriginalSize)*-0.5;
@@ -123,6 +123,13 @@ int32 SSchematicGraphNode::OnPaint(const FPaintArgs& Args, const FGeometry& Allo
 	const bool bIsFadedOut = IsFadedOut();
 	const float FadedOutFactor = bIsFadedOut ? 0.5f : 1.f;
 
+	const int32 FadedOutGroupLayerId = NewLayerId;
+	const int32 FadedOutNodeLayerId = NewLayerId + 100;
+	const int32 FocusedGroupLayerId = NewLayerId + 200;
+	const int32 FocusedNodeLayerId = NewLayerId + 300;
+	int32 NodeLayerId = bIsFadedOut ? FadedOutNodeLayerId : FocusedNodeLayerId;
+	int32 GroupLayerId = bIsFadedOut ? FadedOutGroupLayerId : FocusedGroupLayerId;
+
 	if(const FSchematicGraphGroupNode* GroupNode = Cast<FSchematicGraphGroupNode>(GetNodeData()))
 	{
 		if(GroupNode->GetExpansionState() > SMALL_NUMBER)
@@ -147,11 +154,11 @@ int32 SSchematicGraphNode::OnPaint(const FPaintArgs& Args, const FGeometry& Allo
 			const FVector2d CircleSize = FVector2d::One() * Radius * 2.f;
 			const FVector2d CircleOffset = (CircleSize-OriginalSize)*-0.5;
 
-			NewLayerId++;
+			NodeLayerId++;
 
 			FSlateDrawElement::MakeBox(
 				OutDrawElements,
-				NewLayerId,
+				NodeLayerId,
 				AllottedGeometry.ToPaintGeometry(CircleSize, FSlateLayoutTransform(CircleOffset)),
 				GroupBrush,
 				ESlateDrawEffect::None,
@@ -164,11 +171,11 @@ int32 SSchematicGraphNode::OnPaint(const FPaintArgs& Args, const FGeometry& Allo
 	{
 		for(int32 LayerIndex = 0; LayerIndex < LayerColors.Num(); LayerIndex++)
 		{
-			NewLayerId++;
+			NodeLayerId++;
 
 			FSlateDrawElement::MakeBox(
 				OutDrawElements,
-				NewLayerId,
+				NodeLayerId,
 				AllottedGeometry.ToPaintGeometry(CurSize, FSlateLayoutTransform(SizeOffset)),
 				BrushGetter(GetGuid(), LayerIndex),
 				ESlateDrawEffect::None,
@@ -209,11 +216,11 @@ int32 SSchematicGraphNode::OnPaint(const FPaintArgs& Args, const FGeometry& Allo
 
 			if(const FSlateBrush* BackgroundBrush = CurrentTag->GetBackgroundBrush())
 			{
-				NewLayerId++;
+				NodeLayerId++;
 
 				FSlateDrawElement::MakeBox(
 					OutDrawElements,
-					NewLayerId,
+					NodeLayerId,
 					AllottedGeometry.ToPaintGeometry(TagSize, FSlateLayoutTransform(TagOffset)),
 					BackgroundBrush,
 					ESlateDrawEffect::None,
@@ -223,11 +230,11 @@ int32 SSchematicGraphNode::OnPaint(const FPaintArgs& Args, const FGeometry& Allo
 
 			if(const FSlateBrush* ForegroundBrush = SchematicGraphPanel->GetSchematicGraph()->GetForegroundBrushForTag(CurrentTag))
 			{
-				NewLayerId++;
+				NodeLayerId++;
 
 				FSlateDrawElement::MakeBox(
 					OutDrawElements,
-					NewLayerId,
+					NodeLayerId,
 					AllottedGeometry.ToPaintGeometry(TagSize, FSlateLayoutTransform(TagOffset)),
 					ForegroundBrush,
 					ESlateDrawEffect::None,
@@ -237,12 +244,12 @@ int32 SSchematicGraphNode::OnPaint(const FPaintArgs& Args, const FGeometry& Allo
 
 			if(!LabelSize.IsNearlyZero())
 			{
-				NewLayerId++;
+				NodeLayerId++;
 				const FVector2d LabelOffset = TagCenter - LabelSize * 0.5;
 
 				FSlateDrawElement::MakeText(
 					OutDrawElements,
-					NewLayerId,
+					NodeLayerId,
 					AllottedGeometry.ToPaintGeometry(LabelSize, FSlateLayoutTransform(LabelOffset)),
 					LabelString,
 					SmallFontStyle,
@@ -298,20 +305,20 @@ int32 SSchematicGraphNode::OnPaint(const FPaintArgs& Args, const FGeometry& Allo
 		static const FColor LabelBackgroundColorHex = FColor::FromHex(TEXT("#0F0F0F"));
 		static const FLinearColor LabelBackgroundColor = FLinearColor(LabelBackgroundColorHex) * FLinearColor(1.f, 1.f, 1.f, 0.7f); 
 		
-		NewLayerId++;
+		NodeLayerId++;
 		FSlateDrawElement::MakeBox(
 			OutDrawElements,
-			NewLayerId,
+			NodeLayerId,
 			AllottedGeometry.ToPaintGeometry(LabelBackgroundSize, FSlateLayoutTransform(LabelBackgroundOffset)),
 			LabelBackgroundBrush,
 			ESlateDrawEffect::None,
 			LabelBackgroundColor * FadedOutFactor
 		);
 
-		NewLayerId++;
+		NodeLayerId++;
 		FSlateDrawElement::MakeText(
 			OutDrawElements,
-			NewLayerId,
+			NodeLayerId,
 			AllottedGeometry.ToPaintGeometry(NodeLabelSize, FSlateLayoutTransform(NodeLabelOffset)),
 			NodeLabelString,
 			SmallFontStyle,
@@ -320,7 +327,7 @@ int32 SSchematicGraphNode::OnPaint(const FPaintArgs& Args, const FGeometry& Allo
 		);
 	}
 	
-	return NewLayerId;
+	return NodeLayerId;
 }
 
 void SSchematicGraphNode::OnMouseEnter(const FGeometry& MyGeometry, const FPointerEvent& MouseEvent)
