Schematic: Improve drag and drop behavior for groups

#rnx

diff --git a/Engine/Source/Editor/AnimationEditorWidgets/Private/SchematicGraphPanel/SchematicGraphNode.cpp b/Engine/Source/Editor/AnimationEditorWidgets/Private/SchematicGraphPanel/SchematicGraphNode.cpp
--- a/Engine/Source/Editor/AnimationEditorWidgets/Private/SchematicGraphPanel/SchematicGraphNode.cpp
+++ b/Engine/Source/Editor/AnimationEditorWidgets/Private/SchematicGraphPanel/SchematicGraphNode.cpp
@@ -313,6 +313,7 @@ void FSchematicGraphGroupNode::SetExpanded(bool InExpanded, bool bAutoCloseParen
 	if(!InExpanded && GetExpansionState() < SMALL_NUMBER)
 	{
 		ExpansionState->SetValueAndStop(0.f);
+		LastExpansionState.Reset();
 	}
 	else
 	{
@@ -324,9 +325,14 @@ void FSchematicGraphGroupNode::SetExpanded(bool InExpanded, bool bAutoCloseParen
 				ExpansionState->SetValueAndStop(SMALL_NUMBER * 2.f);
 			}
 		}
-		if(!ExpansionState->GetDelay().IsSet())
+		
+		if(!LastExpansionState.IsSet() || (LastExpansionState.Get(InExpanded) != InExpanded))
 		{
-			ExpansionState->SetDelayOneShot(GetDelayDuration(InExpanded));
+			if(!ExpansionState->GetDelay().IsSet())
+			{
+				ExpansionState->SetDelayOneShot(GetDelayDuration(InExpanded));
+			}
+			LastExpansionState = InExpanded;
 		}
 		ExpansionState->Set(InExpanded ? 1.f : 0.f);
 	}
@@ -338,7 +344,7 @@ void FSchematicGraphGroupNode::SetExpanded(bool InExpanded, bool bAutoCloseParen
 			GroupNode->SetExpanded(InExpanded, bAutoCloseParentGroups);
 		}
 	}
-
+	
 	if(InExpanded)
 	{
 		Model->SetLastExpandedNode(this);
diff --git a/Engine/Source/Editor/AnimationEditorWidgets/Public/SchematicGraphPanel/SchematicGraphNode.h b/Engine/Source/Editor/AnimationEditorWidgets/Public/SchematicGraphPanel/SchematicGraphNode.h
--- a/Engine/Source/Editor/AnimationEditorWidgets/Public/SchematicGraphPanel/SchematicGraphNode.h
+++ b/Engine/Source/Editor/AnimationEditorWidgets/Public/SchematicGraphPanel/SchematicGraphNode.h
@@ -203,8 +203,9 @@ protected:
 
 	float ExpansionRadius = 60.f;
 	FLinearColor ExpansionColor = FLinearColor::White * FLinearColor(1.f, 1.f, 1.f, 0.5f);
-	float EnterDelayDuration = 0.4f;
-	float LeaveDelayDuration = 1.f;
+	float EnterDelayDuration = 0.25f;
+	float LeaveDelayDuration = 0.75f;
+	TOptional<bool> LastExpansionState;
 	TEasingAttributeInterpolator<float>::FSettings AnimationSettings;
 	TSharedPtr<TAnimatedAttribute<float>> ExpansionState;
 
