Schematic: Hide link rendering for 5.4

#rnx

diff --git a/Engine/Source/Editor/AnimationEditorWidgets/Private/SchematicGraphPanel/SSchematicGraphPanel.cpp b/Engine/Source/Editor/AnimationEditorWidgets/Private/SchematicGraphPanel/SSchematicGraphPanel.cpp
--- a/Engine/Source/Editor/AnimationEditorWidgets/Private/SchematicGraphPanel/SSchematicGraphPanel.cpp
+++ b/Engine/Source/Editor/AnimationEditorWidgets/Private/SchematicGraphPanel/SSchematicGraphPanel.cpp
@@ -854,6 +854,7 @@ int32 SSchematicGraphPanel::OnPaint(const FPaintArgs& Args, const FGeometry& All
 	}
 
 	// draw all of the links
+	/*
 	for(const TPair<FGuid, TSharedPtr<FSchematicLinkWidgetInfo>>& Pair : LinkByGuid)
 	{
 		if(GraphData->GetVisibilityForLink(Pair.Key) == ESchematicGraphVisibility::Hidden)
@@ -962,6 +963,7 @@ int32 SSchematicGraphPanel::OnPaint(const FPaintArgs& Args, const FGeometry& All
 				);
 		}
 	}
+	*/
 
 	// Because we paint multiple children, we must track the maximum layer id that they produced in case one of our parents
 	// wants to an overlay for all of its contents.
