Control Rig Schematic: Fix nodes disappearing

#rnx

diff --git a/Engine/Plugins/Animation/ControlRig/Source/ControlRigEditor/Private/Editor/ControlRigSchematicModel.cpp b/Engine/Plugins/Animation/ControlRig/Source/ControlRigEditor/Private/Editor/ControlRigSchematicModel.cpp
--- a/Engine/Plugins/Animation/ControlRig/Source/ControlRigEditor/Private/Editor/ControlRigSchematicModel.cpp
+++ b/Engine/Plugins/Animation/ControlRig/Source/ControlRigEditor/Private/Editor/ControlRigSchematicModel.cpp
@@ -169,7 +169,7 @@ void FControlRigSchematicModel::Tick(float InDeltaTime)
 			}
 			if(Connectors.Num() == 1)
 			{
-				if(const FRigConnectorElement* Connector = Hierarchy->Find<FRigConnectorElement>(Pair.Key))
+				if(const FRigConnectorElement* Connector = Hierarchy->Find<FRigConnectorElement>(Connectors[0]))
 				{
 					if(!Connector->IsPrimary())
 					{
diff --git a/Engine/Source/Editor/AnimationEditorWidgets/Private/SchematicGraphPanel/SSchematicGraphPanel.cpp b/Engine/Source/Editor/AnimationEditorWidgets/Private/SchematicGraphPanel/SSchematicGraphPanel.cpp
--- a/Engine/Source/Editor/AnimationEditorWidgets/Private/SchematicGraphPanel/SSchematicGraphPanel.cpp
+++ b/Engine/Source/Editor/AnimationEditorWidgets/Private/SchematicGraphPanel/SSchematicGraphPanel.cpp
@@ -800,23 +800,21 @@ int32 SSchematicGraphPanel::OnPaint(const FPaintArgs& Args, const FGeometry& All
 	{
 		FArrangedWidget& CurWidget = ArrangedChildren[ChildIndex];
 		const TSharedRef<SSchematicGraphNode> ChildNode = StaticCastSharedRef<SSchematicGraphNode>(CurWidget.Widget);
-
+		
 		const FVector2d NodeCenter = CurWidget.Geometry.GetLocalPositionAtCoordinates({0.5, 0.5});
 		NodeCenterByIndex.Add(NodeCenter);
 		NodeCenterByGuid.Add(ChildNode->GetGuid(), NodeCenter);
 
 		const FSchematicGraphNode* NodeData = ChildNode->GetNodeData();
-
-		ESchematicGraphVisibility::Type NodeVisibility = ESchematicGraphVisibility::Visible;
+		
+		const int32 IndexInPerNodeCache = GuidToNodeCache.FindChecked(NodeData->GetGuid());
+		ESchematicGraphVisibility::Type NodeVisibility = PerNodeCaches[IndexInPerNodeCache].Visibility;
+		
 		if(CurWidget.Geometry.GetLocalSize().IsNearlyZero() ||
 			!FSlateRect::DoRectanglesIntersect( CurWidget.Geometry.GetLayoutBoundingRect(), MyCullingRect ))
 		{
 			NodeVisibility = ESchematicGraphVisibility::Hidden;
 		}
-		else
-		{
-			NodeVisibility = GraphData->GetVisibilityForNode(NodeData);
-		}
 
 		NodeVisibilityByIndex.Add(NodeVisibility);
 		NodeVisibilityByGuid.Add(NodeData->GetGuid(), NodeVisibility);
@@ -1275,6 +1273,7 @@ bool SSchematicGraphPanel::IsAutoScaleEnabledForNode(FGuid InNodeGuid) const
 void SSchematicGraphPanel::UpdatePerNodeCaches(bool bRemoveNodesFromAutoGroups)
 {
 	PerNodeCaches.Reset();
+	GuidToNodeCache.Reset();
 
 	if(GraphData == nullptr)
 	{
@@ -1282,6 +1281,7 @@ void SSchematicGraphPanel::UpdatePerNodeCaches(bool bRemoveNodesFromAutoGroups)
 	}
 
 	PerNodeCaches.Reserve(Children.Num());
+	GuidToNodeCache.Reserve(Children.Num());
 
 	for (int32 i=0; i<Children.Num(); ++i)
 	{
@@ -1297,6 +1297,8 @@ void SSchematicGraphPanel::UpdatePerNodeCaches(bool bRemoveNodesFromAutoGroups)
 					GraphData->RemoveFromParentNode(Node, false);
 				}
 			}
+			Cache.Guid = Node->GetGuid();
+			Cache.Label = Node->GetLabel();
 			Cache.bHasParent = Node->HasParentNode();
 			Cache.Visibility = GraphData->GetVisibilityForNode(Node);
 			Cache.bIsAutoScaling = !Cache.bHasParent && !Widget->bIsBeingDragged && Widget->EnableAutoScale.Get();
@@ -1306,7 +1308,8 @@ void SSchematicGraphPanel::UpdatePerNodeCaches(bool bRemoveNodesFromAutoGroups)
 			Cache.Radius = FMath::Min(NodeSize.X, NodeSize.Y) * 0.5;
 		}
 
-		PerNodeCaches.Add(Cache);
+		const int32 Index = PerNodeCaches.Add(Cache);
+		GuidToNodeCache.Add(Cache.Guid, Index);
 	}
 }
 
diff --git a/Engine/Source/Editor/AnimationEditorWidgets/Public/SchematicGraphPanel/SSchematicGraphPanel.h b/Engine/Source/Editor/AnimationEditorWidgets/Public/SchematicGraphPanel/SSchematicGraphPanel.h
--- a/Engine/Source/Editor/AnimationEditorWidgets/Public/SchematicGraphPanel/SSchematicGraphPanel.h
+++ b/Engine/Source/Editor/AnimationEditorWidgets/Public/SchematicGraphPanel/SSchematicGraphPanel.h
@@ -267,13 +267,17 @@ private:
 	struct FPerNodeCache
 	{
 		FPerNodeCache()
-			: bHasParent(false)
+			: Guid()
+			, Label()
+			, bHasParent(false)
 			, Visibility(ESchematicGraphVisibility::Visible)
 			, bIsAutoScaling(false)
 			, Position(FVector2d::ZeroVector)
 			, Radius(0.0)
 		{}
 
+		FGuid Guid;
+		FText Label;
 		bool bHasParent;
 		ESchematicGraphVisibility::Type Visibility;
 		bool bIsAutoScaling;
@@ -282,6 +286,7 @@ private:
 	};
 
 	TArray<FPerNodeCache> PerNodeCaches;
+	TMap<FGuid, int32> GuidToNodeCache;
 	TMap<uint32, FGuid> GroupNodeGuidByHash;
 
 	TMap<FGuid, TSharedPtr<FSchematicLinkWidgetInfo>> LinkByGuid;
