Schematic: Fix DPI scaling for high res screens

#rnx

diff --git a/Engine/Source/Editor/AnimationEditorWidgets/AnimationEditorWidgets.Build.cs b/Engine/Source/Editor/AnimationEditorWidgets/AnimationEditorWidgets.Build.cs
--- a/Engine/Source/Editor/AnimationEditorWidgets/AnimationEditorWidgets.Build.cs
+++ b/Engine/Source/Editor/AnimationEditorWidgets/AnimationEditorWidgets.Build.cs
@@ -17,6 +17,7 @@ public class AnimationEditorWidgets : ModuleRules
 				"InputCore",
 				"ToolWidgets",
 				"AnimationCore",
+				"ApplicationCore",
 			}
 		);
 
diff --git a/Engine/Source/Editor/AnimationEditorWidgets/Private/SchematicGraphPanel/SSchematicGraphPanel.cpp b/Engine/Source/Editor/AnimationEditorWidgets/Private/SchematicGraphPanel/SSchematicGraphPanel.cpp
--- a/Engine/Source/Editor/AnimationEditorWidgets/Private/SchematicGraphPanel/SSchematicGraphPanel.cpp
+++ b/Engine/Source/Editor/AnimationEditorWidgets/Private/SchematicGraphPanel/SSchematicGraphPanel.cpp
@@ -9,6 +9,7 @@
 #include "Engine/Engine.h"
 #include <SchematicGraphPanel/SchematicGraphStyle.h>
 #include "Fonts/FontMeasure.h"
+#include "HAL/PlatformApplicationMisc.h"
 
 #define LOCTEXT_NAMESPACE "SSchematicGraphPanel"
 
@@ -1018,6 +1019,8 @@ TStatId SSchematicGraphPanel::GetStatId() const
 
 void SSchematicGraphPanel::Tick(float DeltaTime)
 {
+	DPIScale.Reset();
+	
 	if(GraphData)
 	{
 		GraphData->Tick(DeltaTime);
@@ -1181,7 +1184,14 @@ FVector2d SSchematicGraphPanel::GetPositionForNode(FGuid InNodeGuid) const
 		{
 			FVector2d Position = GraphData->GetPositionOffsetForNode(Node);
 			Position += GraphData->GetPositionForNode(Node);
-			return Position;
+
+			if(!DPIScale.IsSet())
+			{
+				const float WidgetX = CachedGeometry.GetAbsolutePosition().X;
+				const float WidgetY = CachedGeometry.GetAbsolutePosition().Y;
+				DPIScale = 1.f / FPlatformApplicationMisc::GetDPIScaleFactorAtPoint(WidgetX, WidgetY);
+			}
+			return DPIScale.GetValue() * Position;
 		}
 	}
 	return FVector2d::ZeroVector;
diff --git a/Engine/Source/Editor/AnimationEditorWidgets/Public/SchematicGraphPanel/SSchematicGraphPanel.h b/Engine/Source/Editor/AnimationEditorWidgets/Public/SchematicGraphPanel/SSchematicGraphPanel.h
--- a/Engine/Source/Editor/AnimationEditorWidgets/Public/SchematicGraphPanel/SSchematicGraphPanel.h
+++ b/Engine/Source/Editor/AnimationEditorWidgets/Public/SchematicGraphPanel/SSchematicGraphPanel.h
@@ -296,6 +296,7 @@ private:
 	mutable TArray<ESchematicGraphVisibility::Type> NodeVisibilityByIndex;
 	mutable TMap<FGuid, ESchematicGraphVisibility::Type> NodeVisibilityByGuid;
 	mutable TOptional<FGuid> DropTarget;
+	mutable TOptional<float> DPIScale;
 };
 
 #endif
\ No newline at end of file
