Control Rig: Fix children nodes of schematic group nodes badly positioned for high-res displays
#rnx

diff --git a/Engine/Source/Editor/AnimationEditorWidgets/Private/SchematicGraphPanel/SSchematicGraphPanel.cpp b/Engine/Source/Editor/AnimationEditorWidgets/Private/SchematicGraphPanel/SSchematicGraphPanel.cpp
--- a/Engine/Source/Editor/AnimationEditorWidgets/Private/SchematicGraphPanel/SSchematicGraphPanel.cpp
+++ b/Engine/Source/Editor/AnimationEditorWidgets/Private/SchematicGraphPanel/SSchematicGraphPanel.cpp
@@ -1185,13 +1185,8 @@ FVector2d SSchematicGraphPanel::GetPositionForNode(FGuid InNodeGuid) const
 			FVector2d Position = GraphData->GetPositionOffsetForNode(Node);
 			Position += GraphData->GetPositionForNode(Node);
 
-			if(!DPIScale.IsSet())
-			{
-				const float WidgetX = CachedGeometry.GetAbsolutePosition().X;
-				const float WidgetY = CachedGeometry.GetAbsolutePosition().Y;
-				DPIScale = 1.f / FPlatformApplicationMisc::GetDPIScaleFactorAtPoint(WidgetX, WidgetY);
-			}
-			return DPIScale.GetValue() * Position;
+			AdjustPositionWithDPIScale(Position);
+			return Position;
 		}
 	}
 	return FVector2d::ZeroVector;
@@ -1238,6 +1233,24 @@ float SSchematicGraphPanel::GetScaleForNode(FGuid InNodeGuid) const
 	return 1.f;
 }
 
+void SSchematicGraphPanel::AdjustPositionWithDPIScale(FVector2d& InOutPosition, bool bInverse) const
+{
+	if(!DPIScale.IsSet())
+	{
+		const float WidgetX = CachedGeometry.GetAbsolutePosition().X;
+		const float WidgetY = CachedGeometry.GetAbsolutePosition().Y;
+		DPIScale = 1.f / 2.f;//FPlatformApplicationMisc::GetDPIScaleFactorAtPoint(WidgetX, WidgetY);
+	}
+	if (bInverse)
+	{
+		InOutPosition /= DPIScale.GetValue();
+	}
+	else
+	{
+		InOutPosition *= DPIScale.GetValue();
+	}
+}
+
 bool SSchematicGraphPanel::IsAutoGroupingEnabled() const
 {
 	if(GraphData)
@@ -1443,8 +1456,13 @@ void SSchematicGraphPanel::UpdateAutoGroupingForNodes()
 		GroupNodeGuidByHash.Add(Pair.Key, Node->GetGuid());
 
 		// move the group node to the right location
-		const FVector2d AveragePosition = NodePositionPerHash.FindChecked(*PositionHash); 
-		Pair.Value->SetPosition(AveragePosition);
+		const FVector2d AveragePosition = NodePositionPerHash.FindChecked(*PositionHash);
+
+		// The average position is already adjusted for DPI scale
+		// The position on the FSchematicGraphNode should not have that correction, we need to undo the scale
+		FVector2d AveragePositionWithoutDPIScale = AveragePosition;
+		AdjustPositionWithDPIScale(AveragePositionWithoutDPIScale, true);
+		Pair.Value->SetPosition(AveragePositionWithoutDPIScale);
 
 		// also update the widget since it has an animated position
 		if(const SSchematicGraphNode* Widget = FindNode(Node->GetGuid()))
diff --git a/Engine/Source/Editor/AnimationEditorWidgets/Public/SchematicGraphPanel/SSchematicGraphPanel.h b/Engine/Source/Editor/AnimationEditorWidgets/Public/SchematicGraphPanel/SSchematicGraphPanel.h
--- a/Engine/Source/Editor/AnimationEditorWidgets/Public/SchematicGraphPanel/SSchematicGraphPanel.h
+++ b/Engine/Source/Editor/AnimationEditorWidgets/Public/SchematicGraphPanel/SSchematicGraphPanel.h
@@ -225,6 +225,7 @@ public:
 	virtual FLinearColor GetColorForNode(FGuid InNodeGuid, int32 InLayerIndex) const;
 	virtual FText GetToolTipForNode(FGuid InNodeGuid) const;
 	virtual float GetScaleForNode(FGuid InNodeGuid) const;
+	void AdjustPositionWithDPIScale(FVector2d& InOutPosition, bool bInverse = false) const;
 
 	virtual bool IsAutoGroupingEnabled() const;
 	virtual float GetAutoGroupingDistance() const;
