Control Rig: Fix children nodes of schematic group nodes badly positioned for high-res displays. Fixed hardcoded scale.
#rnx

diff --git a/Engine/Source/Editor/AnimationEditorWidgets/Private/SchematicGraphPanel/SSchematicGraphPanel.cpp b/Engine/Source/Editor/AnimationEditorWidgets/Private/SchematicGraphPanel/SSchematicGraphPanel.cpp
--- a/Engine/Source/Editor/AnimationEditorWidgets/Private/SchematicGraphPanel/SSchematicGraphPanel.cpp
+++ b/Engine/Source/Editor/AnimationEditorWidgets/Private/SchematicGraphPanel/SSchematicGraphPanel.cpp
@@ -1239,7 +1239,7 @@ void SSchematicGraphPanel::AdjustPositionWithDPIScale(FVector2d& InOutPosition,
 	{
 		const float WidgetX = CachedGeometry.GetAbsolutePosition().X;
 		const float WidgetY = CachedGeometry.GetAbsolutePosition().Y;
-		DPIScale = 1.f / 2.f;//FPlatformApplicationMisc::GetDPIScaleFactorAtPoint(WidgetX, WidgetY);
+		DPIScale = 1.f / FPlatformApplicationMisc::GetDPIScaleFactorAtPoint(WidgetX, WidgetY);
 	}
 	if (bInverse)
 	{
