Control Rig: Prettify labels drawn on the schematic

#rnx

diff --git a/Engine/Plugins/Animation/ControlRig/Source/ControlRigEditor/Private/Editor/ControlRigSchematicModel.cpp b/Engine/Plugins/Animation/ControlRig/Source/ControlRigEditor/Private/Editor/ControlRigSchematicModel.cpp
--- a/Engine/Plugins/Animation/ControlRig/Source/ControlRigEditor/Private/Editor/ControlRigSchematicModel.cpp
+++ b/Engine/Plugins/Animation/ControlRig/Source/ControlRigEditor/Private/Editor/ControlRigSchematicModel.cpp
@@ -75,7 +75,7 @@ const FText& FControlRigSchematicRigElementKeyNode::GetLabel() const
 				for(const FRigElementKey& Connector : Connectors)
 				{
 					static const FText NodeLabelConnectionFormat = LOCTEXT("NodeLabelConnectionFormat", "{0}\nConnection: {1}");
-					const FText ConnectorShortestPath = Hierarchy->GetDisplayNameForUI(Connector, false);
+					const FText ConnectorShortestPath = Hierarchy->GetDisplayNameForUI(Connector, true);
 					LabelText = FText::Format(NodeLabelConnectionFormat, LabelText, ConnectorShortestPath);
 				}
 			}
diff --git a/Engine/Source/Editor/AnimationEditorWidgets/Private/SchematicGraphPanel/SSchematicGraphPanel.cpp b/Engine/Source/Editor/AnimationEditorWidgets/Private/SchematicGraphPanel/SSchematicGraphPanel.cpp
--- a/Engine/Source/Editor/AnimationEditorWidgets/Private/SchematicGraphPanel/SSchematicGraphPanel.cpp
+++ b/Engine/Source/Editor/AnimationEditorWidgets/Private/SchematicGraphPanel/SSchematicGraphPanel.cpp
@@ -295,7 +295,11 @@ int32 SSchematicGraphNode::OnPaint(const FPaintArgs& Args, const FGeometry& Allo
 		FVector2d NodeLabelSize = FontMeasureService->Measure(NodeLabelString, SmallFontStyle);
 
 		const FVector2d NodeBottomCenter = AllottedGeometry.GetLocalSize() * FVector2d(0.5f, 1.f) + FVector2d(0.f, 8.f);
-		const FVector2d NodeLabelOffset = NodeBottomCenter - NodeLabelSize * FVector2d(0.5, 0);
+		FVector2d NodeLabelOffset = NodeBottomCenter - NodeLabelSize * FVector2d(0.5, 0);
+		if(SchematicGraphPanel)
+		{
+			NodeLabelOffset += SchematicGraphPanel->GetNodeLabelOffset();
+		}
 
 		static const FVector2d LabelBackgroundPadding = FVector2d(2.f);
 		const FVector2d LabelBackgroundSize = NodeLabelSize + LabelBackgroundPadding * 2.f;
@@ -325,6 +329,11 @@ int32 SSchematicGraphNode::OnPaint(const FPaintArgs& Args, const FGeometry& Allo
 			ESlateDrawEffect::None,
 			FLinearColor::White * FadedOutFactor
 		);
+
+		if(SchematicGraphPanel)
+		{
+			SchematicGraphPanel->IncrementNodeLabelOffset(FVector2d(0, LabelBackgroundSize.Y));
+		}
 	}
 	
 	return NodeLayerId;
@@ -1028,6 +1037,7 @@ TStatId SSchematicGraphPanel::GetStatId() const
 void SSchematicGraphPanel::Tick(float DeltaTime)
 {
 	DPIScale.Reset();
+	NodeLabelOffset = FVector2d::ZeroVector;
 	
 	if(GraphData)
 	{
@@ -1301,6 +1311,11 @@ bool SSchematicGraphPanel::IsAutoScaleEnabledForNode(FGuid InNodeGuid) const
 	return 0.f;
  }
 
+void SSchematicGraphPanel::IncrementNodeLabelOffset(const FVector2d& InOffset)
+{
+	NodeLabelOffset += InOffset + FVector2d(0,3);
+}
+
 void SSchematicGraphPanel::UpdatePerNodeCaches(bool bRemoveNodesFromAutoGroups)
 {
 	PerNodeCaches.Reset();
diff --git a/Engine/Source/Editor/AnimationEditorWidgets/Public/SchematicGraphPanel/SSchematicGraphPanel.h b/Engine/Source/Editor/AnimationEditorWidgets/Public/SchematicGraphPanel/SSchematicGraphPanel.h
--- a/Engine/Source/Editor/AnimationEditorWidgets/Public/SchematicGraphPanel/SSchematicGraphPanel.h
+++ b/Engine/Source/Editor/AnimationEditorWidgets/Public/SchematicGraphPanel/SSchematicGraphPanel.h
@@ -241,6 +241,9 @@ public:
 	FOnCancelDrag& OnCancelDrag() { return OnCancelDragDelegate; }
 	FOnDrop& OnAcceptDrop() { return OnDropDelegate; }
 
+	const FVector2d& GetNodeLabelOffset() const { return NodeLabelOffset; }
+	void IncrementNodeLabelOffset(const FVector2d& InOffset);
+
 private:
 
 	void UpdatePerNodeCaches(bool bRemoveNodesFromAutoGroups);
@@ -298,6 +301,7 @@ private:
 	mutable TMap<FGuid, ESchematicGraphVisibility::Type> NodeVisibilityByGuid;
 	mutable TOptional<FGuid> DropTarget;
 	mutable TOptional<float> DPIScale;
+	FVector2d NodeLabelOffset = FVector2d::ZeroVector;
 };
 
 #endif
\ No newline at end of file
