Fix for crash in Modular Rig when switching preview mesh

diff --git a/Engine/Plugins/Animation/ControlRig/Source/ControlRigEditor/Private/Editor/ControlRigEditor.cpp b/Engine/Plugins/Animation/ControlRig/Source/ControlRigEditor/Private/Editor/ControlRigEditor.cpp
--- a/Engine/Plugins/Animation/ControlRig/Source/ControlRigEditor/Private/Editor/ControlRigEditor.cpp
+++ b/Engine/Plugins/Animation/ControlRig/Source/ControlRigEditor/Private/Editor/ControlRigEditor.cpp
@@ -290,7 +290,7 @@ FControlRigBaseEditor::FControlRigBaseEditor()
 	, bRefreshDirectionManipulationTargetsRequired(false)
 	, bSchematicViewPortIsHidden(false)
 {
-	
+	SchematicModel = MakeShared<FControlRigSchematicModel>();
 }
 
 FControlRigBaseEditor* FControlRigBaseEditor::GetFromAssetEditorInstance(IAssetEditorInstance* Instance)
@@ -458,9 +458,9 @@ void FControlRigBaseEditor::InitRigVMEditorImpl(const EToolkitMode::Type Mode, c
 
 		if (ControlRigBlueprint->IsModularRig())
 		{
-			SchematicModel.SetEditor(SharedRef());
-			ControlRigBlueprint->OnSetObjectBeingDebugged().AddRaw(&SchematicModel, &FControlRigSchematicModel::OnSetObjectBeingDebugged);
-			ControlRigBlueprint->GetModularRigController()->OnModified().AddRaw(&SchematicModel, &FControlRigSchematicModel::HandleModularRigModified);
+			SchematicModel->SetEditor(SharedRef());
+			ControlRigBlueprint->OnSetObjectBeingDebugged().AddSP(SchematicModel.Get(), &FControlRigSchematicModel::OnSetObjectBeingDebugged);
+			ControlRigBlueprint->GetModularRigController()->OnModified().AddSP(SchematicModel.Get(), &FControlRigSchematicModel::HandleModularRigModified);
 		}
 
 		ControlRigBlueprint->OnRigTypeChanged().AddSP(SharedRef(), &FControlRigBaseEditor::HandleRigTypeChanged);
@@ -475,7 +475,7 @@ void FControlRigBaseEditor::InitRigVMEditorImpl(const EToolkitMode::Type Mode, c
 
 	if(SchematicViewport.IsValid())
 	{
-		SchematicModel.UpdateControlRigContent();
+		SchematicModel->UpdateControlRigContent();
 	}
 }
 
@@ -2692,14 +2692,14 @@ void FControlRigBaseEditor::HandleViewportCreated(const TSharedRef<class IPerson
 			if (Blueprint->IsModularRig())
 			{
 				SchematicViewport = SNew(SSchematicGraphPanel)
-																.GraphData(&SchematicModel)
-																.IsOverlay(true)
-																.PaddingLeft(30)
-																.PaddingRight(30)
-																.PaddingTop(60)
-																.PaddingBottom(60)
-																.PaddingInterNode(5)
-																.Visibility(SharedRef(), &FControlRigBaseEditor::GetSchematicOverlayVisibility)
+										.GraphDataModel(SchematicModel)
+										.IsOverlay(true)
+										.PaddingLeft(30)
+										.PaddingRight(30)
+										.PaddingTop(60)
+										.PaddingBottom(60)
+										.PaddingInterNode(5)
+										.Visibility(SharedRef(), &FControlRigBaseEditor::GetSchematicOverlayVisibility)
 				;
 				InViewport->AddOverlayWidget(SchematicViewport.ToSharedRef());
 			}
@@ -2787,7 +2787,7 @@ void FControlRigBaseEditor::HandleToggleSchematicViewport()
 {
 	if(SchematicViewport.IsValid())
 	{
-		SchematicModel.UpdateControlRigContent();
+		SchematicModel->UpdateControlRigContent();
 		bSchematicViewPortIsHidden = !bSchematicViewPortIsHidden; 
 	}
 }
@@ -3038,9 +3038,9 @@ void FControlRigBaseEditor::UpdateRigVMHostImpl()
 
 		if(IsModularRig() && ControlRig)
 		{
-			if(SchematicModel.ControlRigBlueprint.IsValid())
+			if(SchematicModel->ControlRigBlueprint.IsValid())
 			{
-				SchematicModel.OnSetObjectBeingDebugged(ControlRig);
+				SchematicModel->OnSetObjectBeingDebugged(ControlRig);
 			}
 		}
 	}
@@ -4301,7 +4301,7 @@ void FControlRigBaseEditor::OnHierarchyModified_AnyThread(ERigHierarchyNotificat
 
 	if(SchematicViewport)
 	{
-		SchematicModel.OnHierarchyModified(InNotif, InHierarchy, InSubject);
+		SchematicModel->OnHierarchyModified(InNotif, InHierarchy, InSubject);
 	}
 	
 	FRigElementKey Key;
diff --git a/Engine/Plugins/Animation/ControlRig/Source/ControlRigEditor/Private/Editor/ControlRigEditor.h b/Engine/Plugins/Animation/ControlRig/Source/ControlRigEditor/Private/Editor/ControlRigEditor.h
--- a/Engine/Plugins/Animation/ControlRig/Source/ControlRigEditor/Private/Editor/ControlRigEditor.h
+++ b/Engine/Plugins/Animation/ControlRig/Source/ControlRigEditor/Private/Editor/ControlRigEditor.h
@@ -443,7 +443,7 @@ protected:
 	virtual void SetPreviewInstance(UAnimPreviewInstance* InPreviewInstance) override { PreviewInstance = InPreviewInstance; }
 
 	/** Model for the schematic views */
-	FControlRigSchematicModel SchematicModel;
+	TSharedPtr<FControlRigSchematicModel> SchematicModel;
 
 	/** Delegate to deal with key down evens in the viewport / editor */
 	FPersonaViewportKeyDownDelegate OnKeyDownDelegate;
diff --git a/Engine/Source/Editor/AnimationEditorWidgets/Private/SchematicGraphPanel/SSchematicGraphPanel.cpp b/Engine/Source/Editor/AnimationEditorWidgets/Private/SchematicGraphPanel/SSchematicGraphPanel.cpp
--- a/Engine/Source/Editor/AnimationEditorWidgets/Private/SchematicGraphPanel/SSchematicGraphPanel.cpp
+++ b/Engine/Source/Editor/AnimationEditorWidgets/Private/SchematicGraphPanel/SSchematicGraphPanel.cpp
@@ -184,78 +184,81 @@ int32 SSchematicGraphNode::OnPaint(const FPaintArgs& Args, const FGeometry& Allo
 		}
 	}
 
-	if(NodeData && SchematicGraphPanel && SchematicGraphPanel->GetSchematicGraph())
+	if (NodeData && SchematicGraphPanel)
 	{
-		const float NodeRadius = FMath::Min(CurSize.X, CurSize.Y) * 0.5;
-		
-		const TArray<TSharedPtr<FSchematicGraphTag>>& Tags = NodeData->GetTags();
-		for(const TSharedPtr<FSchematicGraphTag>& TagPtr : Tags)
+		if(TSharedPtr<FSchematicGraphModel> SchematicGraph = SchematicGraphPanel->GetSchematicGraphModel().Pin())
 		{
-			const FSchematicGraphTag* CurrentTag = TagPtr.Get();
-			const ESchematicGraphVisibility::Type TagVisibility = SchematicGraphPanel->GetSchematicGraph()->GetVisibilityForTag(CurrentTag);
-			if(TagVisibility == ESchematicGraphVisibility::Hidden)
+			const float NodeRadius = FMath::Min(CurSize.X, CurSize.Y) * 0.5;
+		
+			const TArray<TSharedPtr<FSchematicGraphTag>>& Tags = NodeData->GetTags();
+			for(const TSharedPtr<FSchematicGraphTag>& TagPtr : Tags)
 			{
-				continue;
-			}
+				const FSchematicGraphTag* CurrentTag = TagPtr.Get();
+				const ESchematicGraphVisibility::Type TagVisibility = SchematicGraph->GetVisibilityForTag(CurrentTag);
+				if(TagVisibility == ESchematicGraphVisibility::Hidden)
+				{
+					continue;
+				}
 
-			FVector2d TagSize = FVector2d::One() * NodeRadius * 1.0f;
-			FVector2d LabelSize = FVector2d::ZeroVector;
-			const FText& Label = SchematicGraphPanel->GetSchematicGraph()->GetLabelForTag(CurrentTag);
-			const FString LabelString = Label.ToString();
+				FVector2d TagSize = FVector2d::One() * NodeRadius * 1.0f;
+				FVector2d LabelSize = FVector2d::ZeroVector;
+				const FText& Label = SchematicGraph->GetLabelForTag(CurrentTag);
+				const FString LabelString = Label.ToString();
 
-			if(!LabelString.IsEmpty())
-			{
-				const TSharedRef< FSlateFontMeasure > FontMeasureService = FSlateApplication::Get().GetRenderer()->GetFontMeasureService();
-				LabelSize = FontMeasureService->Measure(LabelString, SmallFontStyle);
-				TagSize.X = FMath::Max(TagSize.X, LabelSize.X + 4);
-				TagSize.Y = FMath::Max(TagSize.Y, LabelSize.Y + 4);
-			}
+				if(!LabelString.IsEmpty())
+				{
+					const TSharedRef< FSlateFontMeasure > FontMeasureService = FSlateApplication::Get().GetRenderer()->GetFontMeasureService();
+					LabelSize = FontMeasureService->Measure(LabelString, SmallFontStyle);
+					TagSize.X = FMath::Max(TagSize.X, LabelSize.X + 4);
+					TagSize.Y = FMath::Max(TagSize.Y, LabelSize.Y + 4);
+				}
 
-			const FVector2d TagCenter = SizeOffset + CurSize * 0.5 + FVector2d(0, NodeRadius).GetRotated(CurrentTag->GetPlacementAngle());
-			const FVector2d TagOffset = TagCenter - TagSize * 0.5;
+				const FVector2d TagCenter = SizeOffset + CurSize * 0.5 + FVector2d(0, NodeRadius).GetRotated(CurrentTag->GetPlacementAngle());
+				const FVector2d TagOffset = TagCenter - TagSize * 0.5;
 
-			if(const FSlateBrush* BackgroundBrush = CurrentTag->GetBackgroundBrush())
-			{
-				NodeLayerId++;
-
-				FSlateDrawElement::MakeBox(
-					OutDrawElements,
-					NodeLayerId,
-					AllottedGeometry.ToPaintGeometry(TagSize, FSlateLayoutTransform(TagOffset)),
-					BackgroundBrush,
-					ESlateDrawEffect::None,
-					SchematicGraphPanel->GetSchematicGraph()->GetBackgroundColorForTag(CurrentTag) * FadedOutFactor
-				);
-			}
+				if(const FSlateBrush* BackgroundBrush = CurrentTag->GetBackgroundBrush())
+				{
+					NodeLayerId++;
+
+					FSlateDrawElement::MakeBox(
+						OutDrawElements,
+						NodeLayerId,
+						AllottedGeometry.ToPaintGeometry(TagSize, FSlateLayoutTransform(TagOffset)),
+						BackgroundBrush,
+						ESlateDrawEffect::None,
+						SchematicGraph->GetBackgroundColorForTag(CurrentTag) * FadedOutFactor
+					);
+				}
 
-			if(const FSlateBrush* ForegroundBrush = SchematicGraphPanel->GetSchematicGraph()->GetForegroundBrushForTag(CurrentTag))
-			{
-				NodeLayerId++;
-
-				FSlateDrawElement::MakeBox(
-					OutDrawElements,
-					NodeLayerId,
-					AllottedGeometry.ToPaintGeometry(TagSize, FSlateLayoutTransform(TagOffset)),
-					ForegroundBrush,
-					ESlateDrawEffect::None,
-					SchematicGraphPanel->GetSchematicGraph()->GetForegroundColorForTag(CurrentTag) * FadedOutFactor
-				);
-			}
+				if(const FSlateBrush* ForegroundBrush = SchematicGraph->GetForegroundBrushForTag(CurrentTag))
+				{
+					NodeLayerId++;
+
+					FSlateDrawElement::MakeBox(
+						OutDrawElements,
+						NodeLayerId,
+						AllottedGeometry.ToPaintGeometry(TagSize, FSlateLayoutTransform(TagOffset)),
+						ForegroundBrush,
+						ESlateDrawEffect::None,
+						SchematicGraph->GetForegroundColorForTag(CurrentTag) * FadedOutFactor
+					);
+				}
 
-			if(!LabelSize.IsNearlyZero())
-			{
-				NodeLayerId++;
-				const FVector2d LabelOffset = TagCenter - LabelSize * 0.5;
-
-				FSlateDrawElement::MakeText(
-					OutDrawElements,
-					NodeLayerId,
-					AllottedGeometry.ToPaintGeometry(LabelSize, FSlateLayoutTransform(LabelOffset)),
-					LabelString,
-					SmallFontStyle,
-					ESlateDrawEffect::None,
-					SchematicGraphPanel->GetSchematicGraph()->GetLabelColorForTag(CurrentTag) * FadedOutFactor
-				);
+				if(!LabelSize.IsNearlyZero())
+				{
+					NodeLayerId++;
+					const FVector2d LabelOffset = TagCenter - LabelSize * 0.5;
+
+					FSlateDrawElement::MakeText(
+						OutDrawElements,
+						NodeLayerId,
+						AllottedGeometry.ToPaintGeometry(LabelSize, FSlateLayoutTransform(LabelOffset)),
+						LabelString,
+						SmallFontStyle,
+						ESlateDrawEffect::None,
+						SchematicGraph->GetLabelColorForTag(CurrentTag) * FadedOutFactor
+					);
+				}
 			}
 		}
 	}
@@ -441,19 +444,22 @@ FReply SSchematicGraphNode::OnDragDetected(const FGeometry& MyGeometry, const FP
 	TArray<FGuid> DraggedElements = {Guid};
 	if (MouseEvent.IsMouseButtonDown(EKeys::LeftMouseButton) && DraggedElements.Num() > 0)
 	{
-		if(SchematicGraphPanel && SchematicGraphPanel->GetSchematicGraph())
+		if(SchematicGraphPanel)
 		{
-			if(SchematicGraphPanel->GetSchematicGraph()->IsDragSupportedForNode(GetGuid()))
+			if(TSharedPtr<FSchematicGraphModel> SchematicGraph = SchematicGraphPanel->GetSchematicGraphModel().Pin())
 			{
-				bIsBeingDragged = true;
+				if(SchematicGraph->IsDragSupportedForNode(GetGuid()))
+				{
+					bIsBeingDragged = true;
 
-				const FVector2f AbsoluteMousePosition = MouseEvent.GetScreenSpacePosition();
-				const FVector2d LocalMousePosition =  (AbsoluteMousePosition - MyGeometry.GetAbsolutePosition()) / MyGeometry.GetAccumulatedLayoutTransform().GetScale();
-				OffsetDuringDrag = -LocalMousePosition;
+					const FVector2f AbsoluteMousePosition = MouseEvent.GetScreenSpacePosition();
+					const FVector2d LocalMousePosition =  (AbsoluteMousePosition - MyGeometry.GetAbsolutePosition()) / MyGeometry.GetAccumulatedLayoutTransform().GetScale();
+					OffsetDuringDrag = -LocalMousePosition;
 			
-				const TSharedRef<FSchematicGraphNodeDragDropOp> DragDropOp = FSchematicGraphNodeDragDropOp::New({this}, MoveTemp(DraggedElements));
-				OnBeginDragDelegate.ExecuteIfBound(this, DragDropOp.ToSharedPtr());
-				return FReply::Handled().BeginDragDrop(DragDropOp);
+					const TSharedRef<FSchematicGraphNodeDragDropOp> DragDropOp = FSchematicGraphNodeDragDropOp::New({this}, MoveTemp(DraggedElements));
+					OnBeginDragDelegate.ExecuteIfBound(this, DragDropOp.ToSharedPtr());
+					return FReply::Handled().BeginDragDrop(DragDropOp);
+				}
 			}
 		}
 	}
@@ -468,12 +474,15 @@ EVisibility SSchematicGraphNode::GetNodeVisibility() const
 		return EVisibility::HitTestInvisible;
 	}
 
-	if(SchematicGraphPanel && SchematicGraphPanel->GetSchematicGraph())
+	if(SchematicGraphPanel)
 	{
-		ESchematicGraphVisibility::Type Vis = SchematicGraphPanel->GetSchematicGraph()->GetVisibilityForNode(GetGuid());
-		if(Vis == ESchematicGraphVisibility::Hidden)
+		if(TSharedPtr<FSchematicGraphModel> SchematicGraph = SchematicGraphPanel->GetSchematicGraphModel().Pin())
 		{
-			return EVisibility::Hidden;
+			ESchematicGraphVisibility::Type Vis = SchematicGraph->GetVisibilityForNode(GetGuid());
+			if(Vis == ESchematicGraphVisibility::Hidden)
+			{
+				return EVisibility::Hidden;
+			}
 		}
 	}
 
@@ -492,7 +501,7 @@ FReply SSchematicGraphNode::OnMouseButtonDown(const FGeometry& MyGeometry, const
 	{
 		if(SchematicGraphPanel && GetNodeData())
 		{
-			if(const FSchematicGraphModel* Graph = SchematicGraphPanel->GetSchematicGraph())
+			if (TSharedPtr<FSchematicGraphModel> Graph = SchematicGraphPanel->GetSchematicGraphModel().Pin())
 			{
 				FMenuBuilder MenuBuilder(true, nullptr);
 				if(Graph->GetContextMenuForNode(GetNodeData(), MenuBuilder))
@@ -563,9 +572,9 @@ const bool SSchematicGraphNode::IsFadedOut() const
 	return false;
 }
 
-void SSchematicGraphPanel::SetSchematicGraph(FSchematicGraphModel* InGraphData)
+void SSchematicGraphPanel::SetSchematicGraphModel(TWeakPtr<FSchematicGraphModel> InGraphData)
 {
-	if (GraphData)
+	if (TSharedPtr<FSchematicGraphModel> GraphData = GraphDataWeak.Pin())
 	{
 		GraphData->OnNodeAdded().RemoveAll(this);
 		GraphData->OnNodeRemoved().RemoveAll(this);
@@ -574,8 +583,9 @@ void SSchematicGraphPanel::SetSchematicGraph(FSchematicGraphModel* InGraphData)
 		GraphData->OnGraphReset().RemoveAll(this);
 	}
 	
-	GraphData = InGraphData;
-	if (GraphData)
+	GraphDataWeak = InGraphData;
+
+	if (TSharedPtr<FSchematicGraphModel> GraphData = GraphDataWeak.Pin())
 	{
 		GraphData->OnNodeAdded().AddSP(this, &SSchematicGraphPanel::AddNode);
 		GraphData->OnNodeRemoved().AddSP(this, &SSchematicGraphPanel::RemoveNode);
@@ -587,7 +597,7 @@ void SSchematicGraphPanel::SetSchematicGraph(FSchematicGraphModel* InGraphData)
 
 void SSchematicGraphPanel::Construct(const FArguments& InArgs)
 {
-	GraphData = InArgs._GraphData;
+	GraphDataWeak = InArgs._GraphDataModel;
 	bIsOverlay = InArgs._IsOverlay;
 	PaddingLeft = InArgs._PaddingLeft;
 	PaddingRight = InArgs._PaddingRight;
@@ -612,7 +622,7 @@ void SSchematicGraphPanel::Construct(const FArguments& InArgs)
 		SetVisibility(bIsOverlay ? EVisibility::SelfHitTestInvisible : EVisibility::Visible);
 	}
 	
-	if (GraphData)
+	if (TSharedPtr<FSchematicGraphModel> GraphData = GraphDataWeak.Pin())
 	{
 		GraphData->OnNodeAdded().AddSP(this, &SSchematicGraphPanel::AddNode);
 		GraphData->OnNodeRemoved().AddSP(this, &SSchematicGraphPanel::RemoveNode);
@@ -627,7 +637,7 @@ void SSchematicGraphPanel::RebuildPanel()
 {
 	RemoveAllNodes();
 
-	if (GraphData)
+	if (TSharedPtr<FSchematicGraphModel> GraphData = GraphDataWeak.Pin())
 	{
 		for (const TSharedPtr<FSchematicGraphNode>& Node : GraphData->GetNodes())
 		{
@@ -642,7 +652,8 @@ void SSchematicGraphPanel::RebuildPanel()
 
 void SSchematicGraphPanel::AddNode(const FSchematicGraphNode* InNodeToAdd)
 {
-	if(GraphData == nullptr)
+	TSharedPtr<FSchematicGraphModel> GraphData = GraphDataWeak.Pin();
+	if (!GraphData)
 	{
 		return;
 	}
@@ -657,7 +668,7 @@ void SSchematicGraphPanel::AddNode(const FSchematicGraphNode* InNodeToAdd)
 
 	const FGuid Guid = InNodeToAdd->GetGuid();
 	const auto Position = FVector2dAttribute::CreateWithGetter(Vector2DInterpolationSettings, FVector2dAttribute::FGetter::CreateSP(this, &SSchematicGraphPanel::GetPositionForNode, Guid));
-	const auto Size = FVector2dAttribute::CreateWithGetter(Vector2DInterpolationSettings, FVector2dAttribute::FGetter::CreateRaw(GraphData, &FSchematicGraphModel::GetSizeForNode, InNodeToAdd));
+	const auto Size = FVector2dAttribute::CreateWithGetter(Vector2DInterpolationSettings, FVector2dAttribute::FGetter::CreateSP(GraphData.Get(), &FSchematicGraphModel::GetSizeForNode, InNodeToAdd));
 	const auto Scale = FFloatAttribute::CreateWithGetter(FloatInterpolationSettings, FFloatAttribute::FGetter::CreateSP(this, &SSchematicGraphPanel::GetScaleForNode, Guid), 0.f);
 
 	const int32 NumLayers = GraphData ? GraphData->GetNumLayersForNode(InNodeToAdd) : InNodeToAdd->GetNumLayers();
@@ -667,9 +678,13 @@ void SSchematicGraphPanel::AddNode(const FSchematicGraphNode* InNodeToAdd)
 		Colors.Add(FLinearColorAttribute::CreateWithGetter(ColorInterpolationSettings, FLinearColorAttribute::FGetter::CreateSP(this, &SSchematicGraphPanel::GetColorForNode, Guid, LayerIndex)));
 	}
 
-	const TFunction<const FSlateBrush*(const FGuid&, int32)> BrushGetter = [this](const FGuid& InGuid, int32 InLayerIndex) -> const FSlateBrush*
+	const TFunction<const FSlateBrush*(const FGuid&, int32)> BrushGetter = [GraphDataWeak = this->GraphDataWeak](const FGuid& InGuid, int32 InLayerIndex) -> const FSlateBrush*
 	{
-		return GraphData->GetBrushForNode(InGuid, InLayerIndex);
+		if (TSharedPtr<FSchematicGraphModel> GraphData = GraphDataWeak.Pin())
+		{
+			return GraphData->GetBrushForNode(InGuid, InLayerIndex);
+		}
+		return nullptr;
 	};
 
 	const TSharedRef<SSchematicGraphNode> NewNode = SNew(SSchematicGraphNode)
@@ -732,28 +747,26 @@ SSchematicGraphNode* SSchematicGraphPanel::FindNode(const FGuid& InGuid)
 
 void SSchematicGraphPanel::AddLink(const FSchematicGraphLink* InLinkToAdd)
  {
-	if(GraphData == nullptr)
+	if(TSharedPtr<FSchematicGraphModel> GraphData = GraphDataWeak.Pin())
 	{
-		return;
-	}
-	
-	static const TEasingAttributeInterpolator<float>::FSettings FloatInterpolationSettings(EEasingInterpolatorType::CubicEaseOut, 0.1f);
-	static const TEasingAttributeInterpolator<float>::FSettings SlowFloatInterpolationSettings(EEasingInterpolatorType::CubicEaseOut, 0.2f);
-	static const TEasingAttributeInterpolator<FLinearColor>::FSettings ColorInterpolationSettings(EEasingInterpolatorType::CubicEaseOut, 0.2f);
+		static const TEasingAttributeInterpolator<float>::FSettings FloatInterpolationSettings(EEasingInterpolatorType::CubicEaseOut, 0.1f);
+		static const TEasingAttributeInterpolator<float>::FSettings SlowFloatInterpolationSettings(EEasingInterpolatorType::CubicEaseOut, 0.2f);
+		static const TEasingAttributeInterpolator<FLinearColor>::FSettings ColorInterpolationSettings(EEasingInterpolatorType::CubicEaseOut, 0.2f);
 
-	const FGuid Guid = InLinkToAdd->GetGuid();
-	const auto Minimum = FFloatAttribute::CreateWithGetter(SlowFloatInterpolationSettings, FFloatAttribute::FGetter::CreateRaw(GraphData, &FSchematicGraphModel::GetMinimumForLink, InLinkToAdd), 0.5f);
-	const auto Maximum = FFloatAttribute::CreateWithGetter(SlowFloatInterpolationSettings, FFloatAttribute::FGetter::CreateRaw(GraphData, &FSchematicGraphModel::GetMaximumForLink, InLinkToAdd), 0.5f);
-	const auto Color = FLinearColorAttribute::CreateWithGetter(ColorInterpolationSettings, FLinearColorAttribute::FGetter::CreateRaw(GraphData, &FSchematicGraphModel::GetColorForLink, InLinkToAdd));
-	const auto Thickness = FFloatAttribute::CreateWithGetter(FloatInterpolationSettings, FFloatAttribute::FGetter::CreateRaw(GraphData, &FSchematicGraphModel::GetThicknessForLink, InLinkToAdd), 0);
+		const FGuid Guid = InLinkToAdd->GetGuid();
+		const auto Minimum = FFloatAttribute::CreateWithGetter(SlowFloatInterpolationSettings, FFloatAttribute::FGetter::CreateSP(GraphData.Get(), &FSchematicGraphModel::GetMinimumForLink, InLinkToAdd), 0.5f);
+		const auto Maximum = FFloatAttribute::CreateWithGetter(SlowFloatInterpolationSettings, FFloatAttribute::FGetter::CreateSP(GraphData.Get(), &FSchematicGraphModel::GetMaximumForLink, InLinkToAdd), 0.5f);
+		const auto Color = FLinearColorAttribute::CreateWithGetter(ColorInterpolationSettings, FLinearColorAttribute::FGetter::CreateSP(GraphData.Get(), &FSchematicGraphModel::GetColorForLink, InLinkToAdd));
+		const auto Thickness = FFloatAttribute::CreateWithGetter(FloatInterpolationSettings, FFloatAttribute::FGetter::CreateSP(GraphData.Get(), &FSchematicGraphModel::GetThicknessForLink, InLinkToAdd), 0);
 
-	FSchematicLinkWidgetInfo Info;
-	Info.Minimum = Minimum;
-	Info.Maximum = Maximum;
-	Info.Color = Color;
-	Info.Thickness = Thickness;
+		FSchematicLinkWidgetInfo Info;
+		Info.Minimum = Minimum;
+		Info.Maximum = Maximum;
+		Info.Color = Color;
+		Info.Thickness = Thickness;
 
-	LinkByGuid.Add(Guid, MakeShareable(new FSchematicLinkWidgetInfo(Info)));
+		LinkByGuid.Add(Guid, MakeShareable(new FSchematicLinkWidgetInfo(Info)));
+	}
  }
 
  void SSchematicGraphPanel::RemoveLink(const FSchematicGraphLink* InLinkToRemove)
@@ -797,9 +810,10 @@ int32 SSchematicGraphPanel::OnPaint(const FPaintArgs& Args, const FGeometry& All
 		MaxLayerId++;
 	}
 	
+	TSharedPtr<FSchematicGraphModel> GraphData = GraphDataWeak.Pin();
 	if (!GraphData)
 	{
-		return MaxLayerId; 
+		return MaxLayerId;
 	}
 
 	FArrangedChildren ArrangedChildren(EVisibility::Visible);
@@ -1039,7 +1053,8 @@ void SSchematicGraphPanel::Tick(float DeltaTime)
 	DPIScale.Reset();
 	NodeLabelOffset = FVector2d::ZeroVector;
 	
-	if(GraphData)
+	TSharedPtr<FSchematicGraphModel> GraphData = GraphDataWeak.Pin();
+	if (GraphData)
 	{
 		GraphData->Tick(DeltaTime);
 	}
@@ -1084,7 +1099,7 @@ void SSchematicGraphPanel::Tick(float DeltaTime)
 		}
 
 		// update the animation state of the node
-		Widget->EnablePositionAnimation(GraphData->GetPositionAnimationEnabledForNode(Widget->GetGuid()));
+		Widget->EnablePositionAnimation(GraphData ? GraphData->GetPositionAnimationEnabledForNode(Widget->GetGuid()) : false);
 		if(Widget->GetVisibility() != EVisibility::Visible)
 		{
 			continue;
@@ -1176,7 +1191,7 @@ void SSchematicGraphPanel::OnDropEvent(SSchematicGraphNode* Node, const FDragDro
 
 FReply SSchematicGraphPanel::HandleNodeDragDetected(FGuid Guid, const FGeometry& MyGeometry, const FPointerEvent& MouseEvent)
 {
-	if(GraphData)
+	if(TSharedPtr<FSchematicGraphModel> GraphData = GraphDataWeak.Pin())
 	{
 		FGuid ForwardedGuid = Guid;
 		if(GraphData->GetForwardedNodeForDrag(ForwardedGuid))
@@ -1190,21 +1205,23 @@ FReply SSchematicGraphPanel::HandleNodeDragDetected(FGuid Guid, const FGeometry&
 
 FVector2d SSchematicGraphPanel::GetPositionForNode(FGuid InNodeGuid) const
 {
-	const SSchematicGraphNode* NodeWidget = FindNode(InNodeGuid);
-	if(GraphData && NodeWidget)
+	if(const SSchematicGraphNode* NodeWidget = FindNode(InNodeGuid))
 	{
-		if(NodeWidget->PositionDuringDrag.IsSet())
+		if(TSharedPtr<FSchematicGraphModel> GraphData = GraphDataWeak.Pin())
 		{
-			return NodeWidget->PositionDuringDrag.GetValue() + NodeWidget->OffsetDuringDrag.Get(FVector2d::ZeroVector);
-		}
+			if(NodeWidget->PositionDuringDrag.IsSet())
+			{
+				return NodeWidget->PositionDuringDrag.GetValue() + NodeWidget->OffsetDuringDrag.Get(FVector2d::ZeroVector);
+			}
 		
-		if(const FSchematicGraphNode* Node = GraphData->FindNode(InNodeGuid))
-		{
-			FVector2d Position = GraphData->GetPositionOffsetForNode(Node);
-			Position += GraphData->GetPositionForNode(Node);
+			if(const FSchematicGraphNode* Node = GraphData->FindNode(InNodeGuid))
+			{
+				FVector2d Position = GraphData->GetPositionOffsetForNode(Node);
+				Position += GraphData->GetPositionForNode(Node);
 
-			AdjustPositionWithDPIScale(Position);
-			return Position;
+				AdjustPositionWithDPIScale(Position);
+				return Position;
+			}
 		}
 	}
 	return FVector2d::ZeroVector;
@@ -1212,7 +1229,7 @@ FVector2d SSchematicGraphPanel::GetPositionForNode(FGuid InNodeGuid) const
 
 FLinearColor SSchematicGraphPanel::GetColorForNode(FGuid InNodeGuid, int32 InLayerIndex) const
 {
-	if(GraphData)
+	if (TSharedPtr<FSchematicGraphModel> GraphData = GraphDataWeak.Pin())
 	{
 		return GraphData->GetColorForNode(InNodeGuid, InLayerIndex);
 	}
@@ -1221,7 +1238,7 @@ FLinearColor SSchematicGraphPanel::GetColorForNode(FGuid InNodeGuid, int32 InLay
 
 FText SSchematicGraphPanel::GetToolTipForNode(FGuid InNodeGuid) const
 {
-	if(GraphData)
+	if (TSharedPtr<FSchematicGraphModel> GraphData = GraphDataWeak.Pin())
 	{
 		return GraphData->GetToolTipForNode(InNodeGuid);
 	}
@@ -1230,6 +1247,8 @@ FText SSchematicGraphPanel::GetToolTipForNode(FGuid InNodeGuid) const
 
 float SSchematicGraphPanel::GetScaleForNode(FGuid InNodeGuid) const
 {
+	TSharedPtr<FSchematicGraphModel> GraphData = GraphDataWeak.Pin();
+
 	// check if the node may be auto scaled
 	if(const TSharedPtr<SSchematicGraphNode>* NodePtr = NodeByGuid.Find(InNodeGuid))
 	{
@@ -1271,7 +1290,7 @@ void SSchematicGraphPanel::AdjustPositionWithDPIScale(FVector2d& InOutPosition,
 
 bool SSchematicGraphPanel::IsAutoGroupingEnabled() const
 {
-	if(GraphData)
+	if(TSharedPtr<FSchematicGraphModel> GraphData = GraphDataWeak.Pin())
 	{
 		return GraphData->IsAutoGroupingEnabled();
 	}
@@ -1280,7 +1299,7 @@ bool SSchematicGraphPanel::IsAutoGroupingEnabled() const
 
 float SSchematicGraphPanel::GetAutoGroupingDistance() const
 {
-	if(GraphData)
+	if(TSharedPtr<FSchematicGraphModel> GraphData = GraphDataWeak.Pin())
 	{
 		return GraphData->GetAutoGroupingDistance();
 	}
@@ -1289,7 +1308,7 @@ float SSchematicGraphPanel::GetAutoGroupingDistance() const
 
 bool SSchematicGraphPanel::IsAutoScaleEnabledForNode(FGuid InNodeGuid) const
 {
-	if(GraphData)
+	if(TSharedPtr<FSchematicGraphModel> GraphData = GraphDataWeak.Pin())
 	{
 		return GraphData->IsAutoScaleEnabledForNode(InNodeGuid);
 	}
@@ -1298,7 +1317,7 @@ bool SSchematicGraphPanel::IsAutoScaleEnabledForNode(FGuid InNodeGuid) const
 
  float SSchematicGraphPanel::GetMinimumLinkDistanceForNode(FGuid InLinkGuid, bool bIncludeScale) const
  {
-	if(GraphData)
+	if(TSharedPtr<FSchematicGraphModel> GraphData = GraphDataWeak.Pin())
 	{
 		const float MinimumDistance = GraphData->GetMinimumLinkDistanceForNode(InLinkGuid);
 		if(bIncludeScale)
@@ -1321,6 +1340,7 @@ void SSchematicGraphPanel::UpdatePerNodeCaches(bool bRemoveNodesFromAutoGroups)
 	PerNodeCaches.Reset();
 	GuidToNodeCache.Reset();
 
+	TSharedPtr<FSchematicGraphModel> GraphData = GraphDataWeak.Pin();
 	if(GraphData == nullptr)
 	{
 		return;
@@ -1361,6 +1381,7 @@ void SSchematicGraphPanel::UpdatePerNodeCaches(bool bRemoveNodesFromAutoGroups)
 
 void SSchematicGraphPanel::UpdateAutoGroupingForNodes()
 {
+	TSharedPtr<FSchematicGraphModel> GraphData = GraphDataWeak.Pin();
 	if(GraphData == nullptr)
 	{
 		return;
@@ -1499,6 +1520,7 @@ void SSchematicGraphPanel::UpdateAutoGroupingForNodes()
 
 void SSchematicGraphPanel::UpdateAutoScalingForNodes()
 {
+	TSharedPtr<FSchematicGraphModel> GraphData = GraphDataWeak.Pin();
 	if(GraphData == nullptr)
 	{
 		return;
diff --git a/Engine/Source/Editor/AnimationEditorWidgets/Private/SchematicGraphPanel/SchematicViewportTabSummoner.cpp b/Engine/Source/Editor/AnimationEditorWidgets/Private/SchematicGraphPanel/SchematicViewportTabSummoner.cpp
--- a/Engine/Source/Editor/AnimationEditorWidgets/Private/SchematicGraphPanel/SchematicViewportTabSummoner.cpp
+++ b/Engine/Source/Editor/AnimationEditorWidgets/Private/SchematicGraphPanel/SchematicViewportTabSummoner.cpp
@@ -11,7 +11,7 @@ const FName FSchematicViewportTabSummoner::TabID(TEXT("SchematicViewport"));
 
 FSchematicViewportTabSummoner::FSchematicViewportTabSummoner(TSharedPtr<class FAssetEditorToolkit> InHostingApp, const FSchematicViewportArgs& InArgs)
 	: FWorkflowTabFactory(TabID, InHostingApp)
-	, SchematicGraph(InArgs.SchematicGraph)
+	, SchematicGraphWeak(InArgs.SchematicGraphWeak)
 	, OnViewportCreated(InArgs.OnViewportCreated)
 {
 	TabLabel = LOCTEXT("SchematicViewportTabLabel", "Schematic Viewport");
@@ -37,7 +37,7 @@ FTabSpawnerEntry& FSchematicViewportTabSummoner::RegisterTabSpawner(TSharedRef<F
 TSharedRef<SWidget> FSchematicViewportTabSummoner::CreateTabBody(const FWorkflowTabSpawnInfo& Info) const
 {
 	TSharedRef<SSchematicGraphPanel> Viewport = SNew(SSchematicGraphPanel)
-											.GraphData(SchematicGraph)
+											.GraphDataModel(SchematicGraphWeak)
 											.IsOverlay(false)
 											.PaddingLeft(30)
 											.PaddingRight(30)
diff --git a/Engine/Source/Editor/AnimationEditorWidgets/Public/SchematicGraphPanel/SSchematicGraphPanel.h b/Engine/Source/Editor/AnimationEditorWidgets/Public/SchematicGraphPanel/SSchematicGraphPanel.h
--- a/Engine/Source/Editor/AnimationEditorWidgets/Public/SchematicGraphPanel/SSchematicGraphPanel.h
+++ b/Engine/Source/Editor/AnimationEditorWidgets/Public/SchematicGraphPanel/SSchematicGraphPanel.h
@@ -151,7 +151,8 @@ public:
 	SLATE_BEGIN_ARGS(SSchematicGraphPanel)
 	{}
 	SLATE_ARGUMENT(bool, IsOverlay)
-	SLATE_ARGUMENT(FSchematicGraphModel*, GraphData)
+	SLATE_ARGUMENT_DEPRECATED(FSchematicGraphModel*, GraphData, 5.6, "Please, use GraphDataModel instead.")
+	SLATE_ARGUMENT(TWeakPtr<FSchematicGraphModel>, GraphDataModel)
 	SLATE_ARGUMENT(int32, PaddingLeft)
 	SLATE_ARGUMENT(int32, PaddingRight)
 	SLATE_ARGUMENT(int32, PaddingTop)
@@ -180,9 +181,16 @@ public:
 		OnDropDelegate.Unbind();
 	}
 
-	void SetSchematicGraph(FSchematicGraphModel* InGraphData);
-	const FSchematicGraphModel* GetSchematicGraph() const { return GraphData; }
-	FSchematicGraphModel* GetSchematicGraph() { return GraphData; }
+	UE_DEPRECATED(5.6, "SetSchematicGraph with raw pointer has been deprecated, please, use SetSchematicGraphModel with a weak pointer instead.")
+	void SetSchematicGraph(FSchematicGraphModel* InGraphData) {}
+	UE_DEPRECATED(5.6, "GetSchematicGraph with raw pointer has been deprecated, please, use GetSchematicGraphModel with a weak pointer instead.")
+	const FSchematicGraphModel* GetSchematicGraph() const { return nullptr; }
+	UE_DEPRECATED(5.6, "GetSchematicGraph with raw pointer has been deprecated, please, use GetSchematicGraphModel with a weak pointer instead.")
+	FSchematicGraphModel* GetSchematicGraph() { return nullptr; }
+
+	void SetSchematicGraphModel(TWeakPtr<FSchematicGraphModel> InGraphData);
+	const TWeakPtr<const FSchematicGraphModel> GetSchematicGraphModel() const { return GraphDataWeak; }
+	TWeakPtr<FSchematicGraphModel> GetSchematicGraphModel() { return GraphDataWeak; }
 	void Construct(const FArguments& InArgs);
 
 	void RebuildPanel();
@@ -258,7 +266,7 @@ private:
 	int32 PaddingTop = 0;
 	int32 PaddingBottom = 0;
 	int32 PaddingInterNode = 0;
-	FSchematicGraphModel* GraphData = nullptr;
+	TWeakPtr<FSchematicGraphModel> GraphDataWeak;
 	FOnNodeClicked OnNodeClickedDelegate;
 	FOnBeginDrag OnBeginDragDelegate;
 	FOnEndDrag OnEndDragDelegate;
diff --git a/Engine/Source/Editor/AnimationEditorWidgets/Public/SchematicGraphPanel/SchematicViewportTabSummoner.h b/Engine/Source/Editor/AnimationEditorWidgets/Public/SchematicGraphPanel/SchematicViewportTabSummoner.h
--- a/Engine/Source/Editor/AnimationEditorWidgets/Public/SchematicGraphPanel/SchematicViewportTabSummoner.h
+++ b/Engine/Source/Editor/AnimationEditorWidgets/Public/SchematicGraphPanel/SchematicViewportTabSummoner.h
@@ -17,7 +17,7 @@ struct ANIMATIONEDITORWIDGETS_API FSchematicViewportArgs
 	{}
 
 	/** The model which contains the graph to display */
-	FSchematicGraphModel* SchematicGraph;
+	TWeakPtr<FSchematicGraphModel> SchematicGraphWeak;
 	
 	/** Delegate fired when the viewport is created */
 	FOnSchematicViewportCreated OnViewportCreated;
@@ -35,7 +35,7 @@ public:
 	virtual FTabSpawnerEntry& RegisterTabSpawner(TSharedRef<FTabManager> TabManager, const FApplicationMode* CurrentApplicationMode) const;
 	virtual TSharedRef<SWidget> CreateTabBody(const FWorkflowTabSpawnInfo& Info) const override;
 
-	FSchematicGraphModel* SchematicGraph;
+	TWeakPtr<FSchematicGraphModel> SchematicGraphWeak;
 	FOnSchematicViewportCreated OnViewportCreated;
 };
 #endif
\ No newline at end of file
