AnimationEditorWidget: Override Status Widget

#rnx

diff --git a/Engine/Source/Editor/AnimationEditorWidgets/AnimationEditorWidgets.Build.cs b/Engine/Source/Editor/AnimationEditorWidgets/AnimationEditorWidgets.Build.cs
--- a/Engine/Source/Editor/AnimationEditorWidgets/AnimationEditorWidgets.Build.cs
+++ b/Engine/Source/Editor/AnimationEditorWidgets/AnimationEditorWidgets.Build.cs
@@ -18,6 +18,9 @@ public class AnimationEditorWidgets : ModuleRules
 				"ToolWidgets",
 				"AnimationCore",
 				"ApplicationCore",
+				"PropertyEditor",
+				"WidgetRegistration",
+				"ToolMenus",
 			}
 		);
 
@@ -28,7 +31,6 @@ public class AnimationEditorWidgets : ModuleRules
 				"CoreUObject",
 				"UnrealEd",
 				"GraphEditor",
-				"PropertyEditor"
 			}
 		);
 	}
diff --git a/Engine/Source/Editor/AnimationEditorWidgets/Private/Overrides/OverrideStatusDetailsDisplayManager.cpp b/Engine/Source/Editor/AnimationEditorWidgets/Private/Overrides/OverrideStatusDetailsDisplayManager.cpp
new file mode 100644
--- /dev/null
+++ b/Engine/Source/Editor/AnimationEditorWidgets/Private/Overrides/OverrideStatusDetailsDisplayManager.cpp
@@ -0,0 +1,91 @@
+//  Copyright Epic Games, Inc. All Rights Reserved.
+
+#include "Overrides/OverrideStatusDetailsDisplayManager.h"
+#include "Editor/PropertyEditor/Private/SDetailsView.h"
+
+#define LOCTEXT_NAMESPACE "OverrideStatusDetailsDisplayManager"
+
+FOverrideStatusDetailsDisplayManager::~FOverrideStatusDetailsDisplayManager()
+{
+}
+
+bool FOverrideStatusDetailsDisplayManager::ShouldShowCategoryMenu()
+{
+	return false;
+}
+
+const FDetailsViewStyleKey& FOverrideStatusDetailsDisplayManager::GetDetailsViewStyleKey() const
+{
+	return SDetailsView::GetPrimaryDetailsViewStyleKey();
+}
+
+bool FOverrideStatusDetailsDisplayManager::CanConstructPropertyUpdatedWidgetBuilder() const
+{
+	return bIsDisplayingOverrideableObject;
+}
+
+TSharedPtr<FPropertyUpdatedWidgetBuilder> FOverrideStatusDetailsDisplayManager::ConstructPropertyUpdatedWidgetBuilder(const FConstructPropertyUpdatedWidgetBuilderArgs& Args)
+{
+	TSharedPtr<FOverrideStatusDetailsWidgetBuilder> OverridesComboButtonBuilder;
+	if (bIsDisplayingOverrideableObject)
+	{
+		if (Args.InvalidateCachedState.IsBound())
+		{
+			InvalidateCachedState = Args.InvalidateCachedState;
+		}
+
+		OverridesComboButtonBuilder = ConstructOverrideWidgetBuilder( Args );
+	}
+	return OverridesComboButtonBuilder;
+}
+
+void FOverrideStatusDetailsDisplayManager::SetIsDisplayingOverrideableObject(bool InIsDisplayingOverrideableObject)
+{
+	bIsDisplayingOverrideableObject = InIsDisplayingOverrideableObject;
+}
+
+TSharedPtr<FOverrideStatusWidgetMenuBuilder> FOverrideStatusDetailsDisplayManager::GetMenuBuilder(const FOverrideStatusSubject& InSubject) const
+{
+	const TWeakPtr<FDetailsDisplayManager> WeakManager = ConstCastWeakPtr<FDetailsDisplayManager>(AsWeak());
+	return MakeShared<FOverrideStatusWidgetMenuBuilder>(InSubject, StaticCastWeakPtr<FOverrideStatusDetailsDisplayManager>(WeakManager));
+}
+
+TSharedPtr<FOverrideStatusDetailsWidgetBuilder> FOverrideStatusDetailsDisplayManager::ConstructOverrideWidgetBuilder( const FConstructPropertyUpdatedWidgetBuilderArgs& Args )
+{
+	// if we are displaying an overrideable object, create the overrides widget
+	if ( bIsDisplayingOverrideableObject )
+	{
+		TArray<TWeakObjectPtr<const UObject>> Objects;
+		if(Args.Objects)
+		{
+			for(const TWeakObjectPtr<UObject>& ObjectPtr : *Args.Objects)
+			{
+				if(ObjectPtr.IsValid())
+				{
+					Objects.Add(ObjectPtr);
+				}
+			}
+		}
+		if ( !Objects.IsEmpty() )
+		{
+			return MakeShared<FOverrideStatusDetailsWidgetBuilder>(SharedThis(this), Objects, Args.PropertyPath, Args.Category);
+		}
+	}
+	// else return nullptr and use the usual reset to default button
+	return nullptr;
+}
+
+void FOverrideStatusDetailsDisplayManager::OnObjectPropertyChanged(UObject* ObjectBeingModified, FPropertyChangedEvent& PropertyChangedEvent)
+{
+	if (!ObjectBeingModified
+		|| !bIsDisplayingOverrideableObject
+		|| !PropertyChangedEvent.Property
+		|| !InvalidateCachedState.IsBound())
+	{
+		return;
+	}
+
+	InvalidateCachedState.Execute();
+}
+
+#undef LOCTEXT_NAMESPACE
diff --git a/Engine/Source/Editor/AnimationEditorWidgets/Private/Overrides/OverrideStatusDetailsObjectFilter.cpp b/Engine/Source/Editor/AnimationEditorWidgets/Private/Overrides/OverrideStatusDetailsObjectFilter.cpp
new file mode 100644
--- /dev/null
+++ b/Engine/Source/Editor/AnimationEditorWidgets/Private/Overrides/OverrideStatusDetailsObjectFilter.cpp
@@ -0,0 +1,137 @@
+// Copyright Epic Games, Inc. All Rights Reserved.
+
+#include "Overrides/OverrideStatusDetailsObjectFilter.h"
+
+#include "ClassViewerFilter.h"
+#include "ClassViewerModule.h"
+#include "PropertyPath.h"
+
+FOverrideStatusDetailsViewObjectFilter::FOverrideStatusDetailsViewObjectFilter()
+{
+}
+
+void FOverrideStatusDetailsViewObjectFilter::InitializeDisplayManager()
+{
+	OverrideStatusDisplayManager = MakeShared<FOverrideStatusDetailsDisplayManager>();
+	DisplayManager = OverrideStatusDisplayManager;
+
+	CanMergeObjectDelegate.BindStatic(&FOverrideStatusDetailsViewObjectFilter::MergeObjectByClass);
+}
+
+TArray<FDetailsViewObjectRoot> FOverrideStatusDetailsViewObjectFilter::FilterObjects(const TArray<UObject*>& SourceObjects)
+{
+	// this should always be valid while running, ensuring no crash during shutdown
+	if (!OverrideStatusDisplayManager.IsValid() || !OnCanCreateWidget().IsBound())
+	{
+		return {};
+	}
+
+	TArray<TArray<UObject*>> ObjectSetList;
+	TArray<FDetailsViewObjectRoot> Roots;
+	OverrideStatusDisplayManager->SetIsDisplayingOverrideableObject(false);
+
+	for(UObject* SourceObject : SourceObjects)
+	{
+		if(SourceObject)
+		{
+			if(OnCanCreateWidget().Execute(FOverrideStatusSubject(SourceObject)))
+			{
+				if(!ObjectSetList.IsEmpty() && CanMergeObjectDelegate.IsBound())
+				{
+					bool bMergedObjectIntoSet = false;
+					for(TArray<UObject*>& ObjectSet : ObjectSetList)
+					{
+						for(int32 IndexInSet = 0; IndexInSet < ObjectSet.Num(); IndexInSet++)
+						{
+							if(CanMergeObjectDelegate.Execute(SourceObject, ObjectSet[IndexInSet]))
+							{
+								ObjectSet.Add(SourceObject);
+								bMergedObjectIntoSet = true;
+								break;
+							}
+						}
+					}
+					
+					if(bMergedObjectIntoSet)
+					{
+						continue;
+					}
+				}
+				ObjectSetList.Add({SourceObject});
+			}
+		}
+	}
+
+	if(!ObjectSetList.IsEmpty())
+	{
+		for(const TArray<UObject*>& ObjectSet : ObjectSetList)
+		{
+			Roots.Emplace(ObjectSet);
+		}
+		OverrideStatusDisplayManager->SetIsDisplayingOverrideableObject(true);
+	}
+
+	return Roots;
+}
+
+TSharedPtr<FOverrideStatusWidgetMenuBuilder> FOverrideStatusDetailsViewObjectFilter::GetMenuBuilder(const FOverrideStatusSubject& InSubject) const
+{
+	return OverrideStatusDisplayManager->GetMenuBuilder(InSubject);
+}
+
+FOverrideStatus_CanCreateWidget& FOverrideStatusDetailsViewObjectFilter::OnCanCreateWidget()
+{
+	return OverrideStatusDisplayManager->OnCanCreateWidget();
+}
+
+FOverrideStatus_GetStatus& FOverrideStatusDetailsViewObjectFilter::OnGetStatus()
+{
+	return OverrideStatusDisplayManager->OnGetStatus();
+}
+
+FOverrideStatus_OnWidgetClicked& FOverrideStatusDetailsViewObjectFilter::OnWidgetClicked()
+{
+	return OverrideStatusDisplayManager->OnWidgetClicked();
+}
+
+FOverrideStatus_OnGetMenuContent& FOverrideStatusDetailsViewObjectFilter::OnGetMenuContent()
+{
+	return OverrideStatusDisplayManager->OnGetMenuContent();
+}
+
+FOverrideStatus_AddOverride& FOverrideStatusDetailsViewObjectFilter::OnAddOverride()
+{
+	return OverrideStatusDisplayManager->OnAddOverride();
+}
+
+FOverrideStatus_ClearOverride& FOverrideStatusDetailsViewObjectFilter::OnClearOverride()
+{
+	return OverrideStatusDisplayManager->OnClearOverride();
+}
+
+FOverrideStatus_ResetToDefault& FOverrideStatusDetailsViewObjectFilter::OnResetToDefault()
+{
+	return OverrideStatusDisplayManager->OnResetToDefault();
+}
+
+FOverrideStatus_ValueDiffersFromDefault& FOverrideStatusDetailsViewObjectFilter::OnValueDiffersFromDefault()
+{
+	return OverrideStatusDisplayManager->OnValueDiffersFromDefault();
+}
+
+bool FOverrideStatusDetailsViewObjectFilter::MergeObjectByClass(const UObject* InObjectA, const UObject* InObjectB)
+{
+	if(InObjectA && InObjectB)
+	{
+		const UClass* ClassA = InObjectA->GetClass();
+		const UClass* ClassB = InObjectB->GetClass();
+		if(ClassA && ClassB)
+		{
+			if(ClassA->IsChildOf(ClassB) || ClassB->IsChildOf(ClassA))
+			{
+				return true;
+			}
+		}
+	}
+	return false;
+}
diff --git a/Engine/Source/Editor/AnimationEditorWidgets/Private/Overrides/OverrideStatusDetailsWidgetBuilder.cpp b/Engine/Source/Editor/AnimationEditorWidgets/Private/Overrides/OverrideStatusDetailsWidgetBuilder.cpp
new file mode 100644
--- /dev/null
+++ b/Engine/Source/Editor/AnimationEditorWidgets/Private/Overrides/OverrideStatusDetailsWidgetBuilder.cpp
@@ -0,0 +1,168 @@
+//  Copyright Epic Games, Inc. All Rights Reserved.
+
+#include "Overrides/OverrideStatusDetailsWidgetBuilder.h"
+
+#include "DetailsDisplayManager.h"
+#include "Overrides/SOverrideStatusWidget.h"
+#include "Overrides/OverrideStatusDetailsDisplayManager.h"
+#include "SSimpleComboButton.h"
+#include "Styling/SlateBrush.h"
+#include "Brushes/SlateRoundedBoxBrush.h"
+#include "Elements/Common/EditorDataStorageFeatures.h"
+#include "Framework/MetaData/DriverMetaData.h"
+#include "UObject/InstanceDataObjectUtils.h"
+#include "Widgets/Layout/SBox.h"
+
+#define LOCTEXT_NAMESPACE "OverrideStatusDetailsWidgetBuilder"
+
+FOverrideStatusDetailsWidgetBuilder::FOverrideStatusDetailsWidgetBuilder(
+	const TSharedRef<FOverrideStatusDetailsDisplayManager>& InDetailsDisplayManager,
+	const TArray<TWeakObjectPtr<const UObject>>& InObjects,
+	const TSharedPtr<const FPropertyPath>& InPropertyPath,
+	const FName& InCategory)
+	: FPropertyUpdatedWidgetBuilder()
+	, DisplayManager(InDetailsDisplayManager)
+	, Subject(InObjects, InPropertyPath, InCategory)
+{
+}
+
+TSharedPtr<SWidget> FOverrideStatusDetailsWidgetBuilder::GenerateWidget()
+{
+	TSharedPtr<SWidget> StatusWidgets;
+	
+	if (Subject.IsValid())
+	{
+		if(OnCanCreateWidget().IsBound())
+		{
+			if(!OnCanCreateWidget().Execute(Subject))
+			{
+				return SNullWidget::NullWidget;
+			}
+		}
+		
+		TSharedPtr<SBox> Box;
+		TSharedPtr<SOverrideStatusWidget> Button;
+
+		auto GetStatusLambda = [this]() -> EOverrideWidgetStatus::Type
+		{
+			if(OnGetStatus().IsBound())
+			{
+				if(Subject.IsValid())
+				{
+					return OnGetStatus().Execute(Subject);
+				}
+			}
+			return EOverrideWidgetStatus::Undetermined;
+		};
+
+		SAssignNew(Button, SOverrideStatusWidget)
+		.IsHovered(this->IsRowHoveredAttr)
+		.Status_Lambda([GetStatusLambda]() -> EOverrideWidgetStatus::Type
+		{
+			return GetStatusLambda();
+		})
+		.OnClicked_Lambda([this, GetStatusLambda]() -> FReply
+		{
+			if(OnWidgetClicked().IsBound())
+			{
+				if(Subject.IsValid())
+				{
+					return OnWidgetClicked().Execute(Subject, GetStatusLambda());
+				}
+			}
+			else if(Subject.Num() == 1 && Subject.HasPropertyPath() && OnAddOverride().IsBound() && OnClearOverride().IsBound())
+			{
+				if(Subject.IsValid())
+				{
+					if(GetStatusLambda() == EOverrideWidgetStatus::ChangedHere)
+					{
+						return OnClearOverride().Execute(Subject);
+					}
+					if(GetStatusLambda() == EOverrideWidgetStatus::None)
+					{
+						return OnAddOverride().Execute(Subject);
+					}
+					return FReply::Unhandled();
+				}
+			}
+			return FReply::Unhandled();
+		})
+		.MenuContent_Lambda([this, GetStatusLambda]() -> TSharedRef<SWidget>
+		{
+			if(Subject.IsValid())
+			{
+				if(OnGetMenuContent().IsBound())
+				{
+					return OnGetMenuContent().Execute(Subject, GetStatusLambda());
+				}
+
+				// alternatively rely on the default menu builder
+				if(const TSharedPtr<FOverrideStatusWidgetMenuBuilder> MenuBuilder = DisplayManager->GetMenuBuilder(Subject))
+				{
+					if(const TSharedPtr<SWidget> MenuWidget = MenuBuilder->GenerateWidget())
+					{
+						return MenuWidget.ToSharedRef();
+					}
+				}
+			}
+			return SNullWidget::NullWidget;
+		});
+
+		SAssignNew(Box, SBox)
+			.WidthOverride(20.0f)
+			.HeightOverride(20.0f)
+			.HAlign(HAlign_Center)
+			.VAlign(VAlign_Center);
+		Box->SetContent(Button.ToSharedRef());
+
+		StatusWidgets = Box;
+	}
+	
+	return StatusWidgets;
+}
+
+FOverrideStatusDetailsWidgetBuilder::~FOverrideStatusDetailsWidgetBuilder()
+{
+}
+
+FOverrideStatus_CanCreateWidget& FOverrideStatusDetailsWidgetBuilder::OnCanCreateWidget()
+{
+	return DisplayManager->OnCanCreateWidget();
+}
+
+FOverrideStatus_GetStatus& FOverrideStatusDetailsWidgetBuilder::OnGetStatus()
+{
+	return DisplayManager->OnGetStatus();
+}
+
+FOverrideStatus_OnWidgetClicked& FOverrideStatusDetailsWidgetBuilder::OnWidgetClicked()
+{
+	return DisplayManager->OnWidgetClicked();
+}
+
+FOverrideStatus_OnGetMenuContent& FOverrideStatusDetailsWidgetBuilder::OnGetMenuContent()
+{
+	return DisplayManager->OnGetMenuContent();
+}
+
+FOverrideStatus_AddOverride& FOverrideStatusDetailsWidgetBuilder::OnAddOverride()
+{
+	return DisplayManager->OnAddOverride();
+}
+
+FOverrideStatus_ClearOverride& FOverrideStatusDetailsWidgetBuilder::OnClearOverride()
+{
+	return DisplayManager->OnClearOverride();
+}
+
+FOverrideStatus_ResetToDefault& FOverrideStatusDetailsWidgetBuilder::OnResetToDefault()
+{
+	return DisplayManager->OnResetToDefault();
+}
+
+FOverrideStatus_ValueDiffersFromDefault& FOverrideStatusDetailsWidgetBuilder::OnValueDiffersFromDefault()
+{
+	return DisplayManager->OnValueDiffersFromDefault();
+}
+
+#undef LOCTEXT_NAMESPACE
diff --git a/Engine/Source/Editor/AnimationEditorWidgets/Private/Overrides/OverrideStatusSubject.cpp b/Engine/Source/Editor/AnimationEditorWidgets/Private/Overrides/OverrideStatusSubject.cpp
new file mode 100644
--- /dev/null
+++ b/Engine/Source/Editor/AnimationEditorWidgets/Private/Overrides/OverrideStatusSubject.cpp
@@ -0,0 +1,131 @@
+//  Copyright Epic Games, Inc. All Rights Reserved.
+
+#include "Overrides/OverrideStatusSubject.h"
+
+#include "PropertyPath.h"
+
+#define LOCTEXT_NAMESPACE "OverrideStatusSubject"
+
+FOverrideStatusSubject::FOverrideStatusSubject(const TArray<TWeakObjectPtr<const UObject>>& InObjects,
+	const TSharedPtr<const FPropertyPath>& InPropertyPath, const FName& InCategory)
+	: Objects(InObjects)
+	, PropertyPath(InPropertyPath)
+	, Category(InCategory)
+{
+}
+
+FOverrideStatusSubject::FOverrideStatusSubject(const TWeakObjectPtr<const UObject>& InObject,
+	const TSharedPtr<const FPropertyPath>& InPropertyPath, const FName& InCategory)
+	: Objects({InObject})
+	, PropertyPath(InPropertyPath)
+	, Category(InCategory)
+{
+}
+
+FOverrideStatusSubject::FOverrideStatusSubject(const UObject* InObject, const TSharedPtr<const FPropertyPath>& InPropertyPath,
+	const FName& InCategory)
+	: Objects({TWeakObjectPtr<const UObject>(InObject)})
+	, PropertyPath(InPropertyPath)
+	, Category(InCategory)
+{
+}
+
+bool FOverrideStatusSubject::IsValid() const
+{
+	for(const TWeakObjectPtr<const UObject>& Object : Objects)
+	{
+		if(!Object.IsValid())
+		{
+			return false;
+		}
+	}
+	return !Objects.IsEmpty();
+}
+
+int32 FOverrideStatusSubject::Num() const
+{
+	return Objects.Num();
+}
+
+const UObject* FOverrideStatusSubject::operator[](int32 InIndex) const
+{
+	if(Objects[InIndex].IsValid())
+	{
+		return Objects[InIndex].Get();
+	}
+	return nullptr;
+}
+
+bool FOverrideStatusSubject::HasPropertyPath() const
+{
+	return PropertyPath.IsValid();
+}
+
+const TSharedPtr<const FPropertyPath>& FOverrideStatusSubject::GetPropertyPath() const
+{
+	return PropertyPath;
+}
+
+const FString& FOverrideStatusSubject::GetPropertyPathString(const TCHAR* Separator) const
+{
+	if(!PropertyPath.IsValid())
+	{
+		static const FString EmptyString;
+		return EmptyString;
+	}
+	
+	if(LastSeparator.IsSet() && LastPropertyPathString.IsSet())
+	{
+		if(LastSeparator->Equals(Separator))
+		{
+			return LastPropertyPathString.GetValue();
+		}
+	}
+	
+	LastSeparator = FString(Separator);
+	LastPropertyPathString = PropertyPath->ToString(Separator);
+	return LastPropertyPathString.GetValue();
+}
+
+bool FOverrideStatusSubject::HasCategory() const
+{
+	return !Category.IsNone();
+}
+
+const FName& FOverrideStatusSubject::GetCategory() const
+{
+	return Category;
+}
+
+bool FOverrideStatusSubject::Contains(const UObject* InObject) const
+{
+	for(const TWeakObjectPtr<const UObject>& Object : Objects)
+	{
+		if(Object.IsValid())
+		{
+			if(Object.Get() == InObject)
+			{
+				return true;
+			}
+		}
+	}
+	return false;
+}
+
+int32 FOverrideStatusSubject::Find(const UObject* InObject) const
+{
+	for(int32 Index = 0; Index < Objects.Num(); Index++)
+	{
+		const TWeakObjectPtr<const UObject>& Object = Objects[Index];
+		if(Object.IsValid())
+		{
+			if(Object.Get() == InObject)
+			{
+				return Index;
+			}
+		}
+	}
+	return INDEX_NONE;
+}
+
+#undef LOCTEXT_NAMESPACE
diff --git a/Engine/Source/Editor/AnimationEditorWidgets/Private/Overrides/OverrideStatusWidgetMenuBuilder.cpp b/Engine/Source/Editor/AnimationEditorWidgets/Private/Overrides/OverrideStatusWidgetMenuBuilder.cpp
new file mode 100644
--- /dev/null
+++ b/Engine/Source/Editor/AnimationEditorWidgets/Private/Overrides/OverrideStatusWidgetMenuBuilder.cpp
@@ -0,0 +1,223 @@
+//  Copyright Epic Games, Inc. All Rights Reserved.
+
+#include "Overrides/OverrideStatusWidgetMenuBuilder.h"
+
+#include "Overrides/OverrideStatusDetailsDisplayManager.h"
+#include "ToolMenuSection.h"
+#include "ToolMenus.h"
+#include "UObject/PropertyBagRepository.h"
+
+#define LOCTEXT_NAMESPACE "OverrideStatusWidgetMenuBuilder"
+
+FOverrideStatusWidgetMenuBuilder::FOverrideStatusWidgetMenuBuilder(const FOverrideStatusSubject& InSubject, const TWeakPtr<FOverrideStatusDetailsDisplayManager>& InDisplayManager)
+	: FToolElementRegistrationArgs("FOverrideStatusWidgetMenuBuilder")
+	, Subject(InSubject)
+{
+	if(InDisplayManager.IsValid())
+	{
+		TSharedPtr<FOverrideStatusDetailsDisplayManager> Manager = InDisplayManager.Pin();
+		FOverrideStatus_GetStatus GetStatusDelegate = Manager->OnGetStatus();
+		StatusAttribute.BindLambda([GetStatusDelegate, this]() -> EOverrideWidgetStatus::Type
+		{
+			if(GetStatusDelegate.IsBound())
+			{
+				if(Subject.IsValid())
+				{
+					return GetStatusDelegate.Execute(Subject);
+				}
+			}
+			return EOverrideWidgetStatus::Undetermined;
+		});
+
+		AddOverrideDelegate = Manager->OnAddOverride();
+		ClearOverrideDelegate = Manager->OnClearOverride();
+		ResetToDefaultDelegate = Manager->OnResetToDefault();
+		ValueDiffersFromDefaultDelegate = Manager->OnValueDiffersFromDefault();
+	}
+	else
+	{
+		StatusAttribute = EOverrideWidgetStatus::Undetermined;
+	}
+}
+
+FOverrideStatusWidgetMenuBuilder::~FOverrideStatusWidgetMenuBuilder()
+{
+}
+
+EOverrideWidgetStatus::Type FOverrideStatusWidgetMenuBuilder::GetStatus() const
+{
+	return StatusAttribute.Get(EOverrideWidgetStatus::Undetermined);
+}
+
+void FOverrideStatusWidgetMenuBuilder::AddOverride()
+{
+	if(AddOverrideDelegate.IsBound() && Subject.IsValid())
+	{
+		(void)AddOverrideDelegate.Execute(Subject);
+	}
+}
+
+bool FOverrideStatusWidgetMenuBuilder::CanAddOverride() const
+{
+	if(!Subject.IsValid())
+	{
+		return false;
+	}
+	static const TSet<EOverrideWidgetStatus::Type> ValidStatusList = {
+		EOverrideWidgetStatus::None,
+		EOverrideWidgetStatus::Undetermined,
+		EOverrideWidgetStatus::ChangedInside,
+	};
+	return AddOverrideDelegate.IsBound() && ValidStatusList.Contains(GetStatus());
+}
+
+FOverrideStatus_AddOverride& FOverrideStatusWidgetMenuBuilder::OnAddOverride()
+{
+	return AddOverrideDelegate;
+}
+
+void FOverrideStatusWidgetMenuBuilder::ClearOverride()
+{
+	if(ClearOverrideDelegate.IsBound() && Subject.IsValid())
+	{
+		(void)ClearOverrideDelegate.Execute(Subject);
+	}
+}
+
+bool FOverrideStatusWidgetMenuBuilder::CanClearOverride() const
+{
+	if(!Subject.IsValid())
+	{
+		return false;
+	}
+	static const TSet<EOverrideWidgetStatus::Type> ValidStatusList = {
+		EOverrideWidgetStatus::ChangedHere,
+		EOverrideWidgetStatus::ChangedInside,
+		EOverrideWidgetStatus::Undetermined,
+	};
+	return ClearOverrideDelegate.IsBound() && ValidStatusList.Contains(GetStatus());
+}
+
+FOverrideStatus_ClearOverride& FOverrideStatusWidgetMenuBuilder::OnClearOverride()
+{
+	return ClearOverrideDelegate;
+}
+
+void FOverrideStatusWidgetMenuBuilder::ResetToDefault()
+{
+	if(ResetToDefaultDelegate.IsBound() && Subject.IsValid())
+	{
+		(void)ResetToDefaultDelegate.Execute(Subject);
+	}
+}
+
+bool FOverrideStatusWidgetMenuBuilder::CanResetToDefault() const
+{
+	if(!Subject.IsValid())
+	{
+		return false;
+	}
+	if(!ValueDiffersFromDefaultDelegate.IsBound())
+	{
+		return false;
+	}
+	if(!ValueDiffersFromDefaultDelegate.Execute(Subject))
+	{
+		return false;
+	}
+	static const TSet<EOverrideWidgetStatus::Type> ValidStatusList = {
+		EOverrideWidgetStatus::ChangedHere,
+		EOverrideWidgetStatus::ChangedInside,
+		EOverrideWidgetStatus::Undetermined,
+	};
+	return ResetToDefaultDelegate.IsBound() && ValidStatusList.Contains(GetStatus());
+}
+
+FOverrideStatus_ResetToDefault& FOverrideStatusWidgetMenuBuilder::OnResetToDefault()
+{
+	return ResetToDefaultDelegate;
+}
+
+FOverrideStatus_ValueDiffersFromDefault& FOverrideStatusWidgetMenuBuilder::OnValueDiffersFromDefault()
+{
+	return ValueDiffersFromDefaultDelegate;
+}
+
+void FOverrideStatusWidgetMenuBuilder::InitializeMenu()
+{
+	static const FName CategoryMenuName = "CategoryMenu";
+	static const FName PropertyMenuName = "PropertyMenu";
+	static const FName Overrides = FName(TEXT("Overrides"));
+	const FName ActiveMenu = Subject.HasPropertyPath() ? PropertyMenuName : CategoryMenuName;
+
+	ToolMenu = UToolMenus::Get()->
+						RegisterMenu(ActiveMenu,
+							NAME_None,
+							EMultiBoxType::Menu, false);
+
+	FToolMenuSection& OverrideSection = ToolMenu->AddSection(Overrides);
+
+	TSharedPtr<FOverrideStatusWidgetMenuBuilder> ThisSharedPtr = SharedThis(this);
+
+	OverrideSection.AddMenuEntry(
+		TEXT("AddOverride"),
+		LOCTEXT("AddOverride", "Add Override"),
+		LOCTEXT("AddOverrideTooltip", "Adds an override to this property."),
+		FSlateIcon(),
+		FUIAction(
+			FExecuteAction::CreateLambda([ThisSharedPtr]
+			{
+				ThisSharedPtr->AddOverride();
+			}),
+			FCanExecuteAction::CreateLambda([ThisSharedPtr]
+			{
+				return ThisSharedPtr->CanAddOverride();
+			}))
+	);
+
+	OverrideSection.AddMenuEntry(
+		TEXT("ClearOverride"),
+		LOCTEXT("ClearOverride", "Clear Override"),
+		LOCTEXT("ClearOverrideTooltip", "Clears the override to this property and resets the value."),
+		FSlateIcon(),
+		FUIAction(
+			FExecuteAction::CreateLambda([ThisSharedPtr]
+			{
+				ThisSharedPtr->ClearOverride();
+			}),
+			FCanExecuteAction::CreateLambda([ThisSharedPtr]
+			{
+				return ThisSharedPtr->CanClearOverride();
+			}))
+	);
+
+	OverrideSection.AddMenuEntry(
+		TEXT("ResetToDefault"),
+		LOCTEXT("ResetToDefault", "Reset to default"),
+		LOCTEXT("ResetToDefaultTooltip", "Revert the value to its default but keeps the override."),
+		FSlateIcon(),
+		FUIAction(
+			FExecuteAction::CreateLambda([ThisSharedPtr]
+			{
+				ThisSharedPtr->ResetToDefault();
+			}),
+			FCanExecuteAction::CreateLambda([ThisSharedPtr]
+		{
+			return ThisSharedPtr->CanResetToDefault();
+		}))
+	);
+
+	ToolMenu->bShouldCloseWindowAfterMenuSelection = true;
+	ToolMenu->bCloseSelfOnly = true;
+}
+
+TSharedPtr<SWidget> FOverrideStatusWidgetMenuBuilder::GenerateWidget()
+{
+	InitializeMenu();
+	
+	return ToolMenu.IsValid() ?
+		UToolMenus::Get()->GenerateWidget(ToolMenu.Get()) :
+		SNullWidget::NullWidget;
+}
+
+#undef LOCTEXT_NAMESPACE
diff --git a/Engine/Source/Editor/AnimationEditorWidgets/Private/Overrides/SOverrideListWidget.cpp b/Engine/Source/Editor/AnimationEditorWidgets/Private/Overrides/SOverrideListWidget.cpp
new file mode 100644
--- /dev/null
+++ b/Engine/Source/Editor/AnimationEditorWidgets/Private/Overrides/SOverrideListWidget.cpp
@@ -0,0 +1,89 @@
+//  Copyright Epic Games, Inc. All Rights Reserved.
+
+#include "Overrides/SOverrideListWidget.h"
+
+#include "Widgets/SBoxPanel.h"
+
+#define LOCTEXT_NAMESPACE "SOverrideListWidget"
+
+SLATE_IMPLEMENT_WIDGET(SOverrideListWidget)
+
+void SOverrideListWidget::PrivateRegisterAttributes(FSlateAttributeInitializer& AttributeInitializer)
+{
+}
+
+SOverrideListWidget::SOverrideListWidget()
+{
+}
+
+SOverrideListWidget::~SOverrideListWidget()
+{
+};
+
+void SOverrideListWidget::Construct(const FArguments& InArgs)
+{
+	SubjectsHashAttribute = InArgs._SubjectsHash;
+	SubjectsAttribute = InArgs._Subjects;
+	GetStatusDelegate = InArgs._OnGetStatus;
+	ClearOverrideDelegate = InArgs._OnClearOverride;
+
+	ChildSlot
+	[
+		SNew(SVerticalBox)
+		+ SVerticalBox::Slot()
+		.FillHeight(1)
+		.HAlign(HAlign_Left)
+		.VAlign(VAlign_Top)
+		.Padding(4)
+		[
+			SNew(SHorizontalBox)
+			+ SHorizontalBox::Slot()
+			.FillWidth(1)
+			.HAlign(HAlign_Left)
+			.VAlign(VAlign_Top)
+			.Padding(0)
+			[
+				SAssignNew(TextBlock, STextBlock)
+			]
+		]
+	];
+
+	SetCanTick(true);
+}
+
+void SOverrideListWidget::Tick(
+	const FGeometry& AllottedGeometry,
+	const double InCurrentTime,
+	const float InDeltaTime)
+{
+	TRACE_CPUPROFILER_EVENT_SCOPE(SOverrideListWidget::Tick);
+	
+	Super::Tick(AllottedGeometry, InCurrentTime, InDeltaTime);
+
+	const uint32 SubjectsHash = SubjectsHashAttribute.Get();
+	if(LastHash.Get(UINT32_MAX) != SubjectsHash)
+	{
+		LastHash = SubjectsHash;
+		
+		const TArray<FOverrideStatusSubject>& Subjects = SubjectsAttribute.Get();
+
+		FText Content;
+
+		for(const FOverrideStatusSubject& Subject : Subjects)
+		{
+			if(Subject.IsValid())
+			{
+				Content = FText::Format(
+					LOCTEXT("SubjectFormat", "{0}\n{1} {2}"),
+					Content,
+					FText::FromName(Subject[0]->GetFName()),
+					FText::FromString(Subject.GetPropertyPathString())
+				);
+			}
+		}
+
+		TextBlock->SetText(Content);
+	}
+}
+
+#undef LOCTEXT_NAMESPACE
diff --git a/Engine/Source/Editor/AnimationEditorWidgets/Private/Overrides/SOverrideStatusWidget.cpp b/Engine/Source/Editor/AnimationEditorWidgets/Private/Overrides/SOverrideStatusWidget.cpp
new file mode 100644
--- /dev/null
+++ b/Engine/Source/Editor/AnimationEditorWidgets/Private/Overrides/SOverrideStatusWidget.cpp
@@ -0,0 +1,302 @@
+//  Copyright Epic Games, Inc. All Rights Reserved.
+
+#include "Overrides/SOverrideStatusWidget.h"
+
+#include "Components/HorizontalBox.h"
+#include "SSimpleComboButton.h"
+#include "Framework/Application/SlateApplication.h"
+#include "Styling/StyleColors.h"
+#include "Widgets/Images/SImage.h"
+
+#define LOCTEXT_NAMESPACE "SOverrideStatusWidget"
+
+SLATE_IMPLEMENT_WIDGET(SOverrideStatusWidget)
+
+void SOverrideStatusWidget::PrivateRegisterAttributes(FSlateAttributeInitializer& AttributeInitializer)
+{
+}
+
+SOverrideStatusWidget::SOverrideStatusWidget()
+	: WidgetStyle(nullptr)
+{
+}
+
+SOverrideStatusWidget::~SOverrideStatusWidget()
+{
+};
+
+void SOverrideStatusWidget::Construct(const FArguments& InArgs)
+{
+	StatusAttribute = InArgs._Status;
+	OnGetMenuContent = InArgs._MenuContent;
+	OnClicked = InArgs._OnClicked;
+	bAlwaysUpdateOnTick = false;
+	bTickIsPending = false;
+
+	if(InArgs._IsHovered.IsBound())
+	{
+		SetHover(InArgs._IsHovered);
+	}
+	
+	if(InArgs._DefaultStyle.IsSet())
+	{
+		DefaultWidgetStyle = InArgs._DefaultStyle.GetValue();
+	}
+	else
+	{
+		if(const FOverrideStatusWidgetStyle* StyleFromStatus = GetStyleFromStatus(EOverrideWidgetStatus::Undetermined))
+		{
+			DefaultWidgetStyle = *StyleFromStatus;
+		}
+	}
+	
+	WidgetStyle = nullptr;
+	OnGetStyle = InArgs._Style;
+	OnGetTooltip = InArgs._Tooltip;
+
+	if(!OnGetStyle.IsBound())
+	{
+		OnGetStyle.BindStatic(&SOverrideStatusWidget::GetStyleFromStatus_Fallback);
+		bAlwaysUpdateOnTick = StatusAttribute.IsBound();
+	}
+
+	ChildSlot
+	[
+		SAssignNew(Image, SImage)
+	];
+}
+
+FReply SOverrideStatusWidget::OnMouseButtonUp(const FGeometry& MyGeometry, const FPointerEvent& MouseEvent)
+{
+	return FReply::Handled();
+}
+
+void SOverrideStatusWidget::OnMouseEnter(const FGeometry& MyGeometry, const FPointerEvent& MouseEvent)
+{
+	SCompoundWidget::OnMouseEnter(MyGeometry, MouseEvent);
+	if(!GetHoveredAttribute().IsBound())
+	{
+		SetHover(true);
+	}
+	QueueRepaint();
+}
+
+void SOverrideStatusWidget::OnMouseLeave(const FPointerEvent& MouseEvent)
+{
+	SCompoundWidget::OnMouseLeave(MouseEvent);
+	if(!GetHoveredAttribute().IsBound())
+	{
+		SetHover(false);
+	}
+	QueueRepaint();
+}
+
+void SOverrideStatusWidget::QueueRepaint()
+{
+	Invalidate(EInvalidateWidgetReason::Paint);
+	SetCanTick(true);
+	bTickIsPending = true;
+}
+
+FReply SOverrideStatusWidget::OnMouseButtonDown(const FGeometry& MyGeometry, const FPointerEvent& MouseEvent)
+{
+	if (MouseEvent.GetEffectingButton() != EKeys::LeftMouseButton)
+	{
+		return FReply::Unhandled();
+	}
+	return HandleClick();
+}
+
+FReply SOverrideStatusWidget::HandleClick()
+{
+	// Create the overrides menu at the cursor location
+	if(OnClicked.IsBound())
+	{
+		const FReply ReplyFromDelegate = OnClicked.Execute();
+		if(ReplyFromDelegate.IsEventHandled())
+		{
+			return ReplyFromDelegate;
+		}
+	}
+	if (OnGetMenuContent.IsBound())
+	{
+		const TSharedRef<SWidget> Content = OnGetMenuContent.Execute();
+		if(Content != SNullWidget::NullWidget)
+		{
+			FSlateApplication::Get().PushMenu(
+				SharedThis(this),
+				FWidgetPath(),
+				Content,
+				FSlateApplication::Get().GetCursorPos(),
+				FPopupTransitionEffect(FPopupTransitionEffect::TypeInPopup));
+			return FReply::Handled();
+		}
+	}
+	
+	return FReply::Unhandled();
+}
+
+EOverrideWidgetStatus::Type SOverrideStatusWidget::GetStatus() const
+{
+	return StatusAttribute.Get(EOverrideWidgetStatus::Undetermined);
+}
+
+void SOverrideStatusWidget::SetStatus(EOverrideWidgetStatus::Type InStatus)
+{
+	StatusAttribute = InStatus;
+	QueueRepaint();
+}
+
+const FOverrideStatusWidgetStyle* SOverrideStatusWidget::GetStyleFromStatus_Fallback(EOverrideWidgetStatus::Type InStatus)
+{
+	return GetStyleFromStatus(InStatus);
+}
+
+void SOverrideStatusWidget::Tick(
+	const FGeometry& AllottedGeometry,
+	const double InCurrentTime,
+	const float InDeltaTime)
+{
+	TRACE_CPUPROFILER_EVENT_SCOPE(SOverrideStatusWidget::Tick);
+	
+	if(!bAlwaysUpdateOnTick)
+	{
+		// Only need to tick on first run or if the data changes which trigger a repaint
+		// Paint will tick us.
+		SetCanTick(false);
+	}
+	else
+	{
+		bTickIsPending = true;
+	}
+
+	if(bTickIsPending)
+	{
+		// Figure out what type of Icon we need
+		const FOverrideStatusWidgetStyle* Style = nullptr;
+		if (OnGetStyle.IsBound())
+		{
+			Style = OnGetStyle.Execute(GetStatus());
+		}
+
+		const FOverrideStatusWidgetStyle* NewWidgetStyle = Style ? Style : &DefaultWidgetStyle;
+		if(NewWidgetStyle != WidgetStyle)
+		{
+			WidgetStyle = NewWidgetStyle;
+			
+			Image->SetDesiredSizeOverride(WidgetStyle->IconSize);
+			Image->SetToolTipText(WidgetStyle->Tooltip);
+
+			bLastHovered.Reset();
+		}
+
+		const bool bIsHovered = IsHovered();
+		if (bIsHovered != bLastHovered.Get(!bIsHovered))
+		{
+			bLastHovered = bIsHovered;
+			if(bIsHovered)
+			{
+				Image->SetImage(WidgetStyle->HoveredIcon);
+				Image->SetColorAndOpacity(WidgetStyle->HoveredColor);
+			}
+			else
+			{
+				Image->SetImage(WidgetStyle->Icon);
+				Image->SetColorAndOpacity(WidgetStyle->Color);
+			}
+		}
+
+		bTickIsPending = false;
+	}
+	
+	Super::Tick(AllottedGeometry, InCurrentTime, InDeltaTime);
+}
+
+const FOverrideStatusWidgetStyle* SOverrideStatusWidget::GetStyleFromStatus(EOverrideWidgetStatus::Type InStatus)
+{
+	static const TMap<EOverrideWidgetStatus::Type, FOverrideStatusWidgetStyle> StyleMap =
+		{
+			{ // No override
+				EOverrideWidgetStatus::None,
+				{
+					.Icon = FAppStyle::GetBrush(TEXT("DetailsView.OverrideNone")),
+					.HoveredIcon = FAppStyle::GetBrush(TEXT("DetailsView.OverrideNone.Hovered")),
+					.Tooltip = LOCTEXT("NoOverridePropertyToolTip", "No override."),
+					.IconSize = FVector2D(16, 16),
+				}
+			},
+			{ // Undetermined
+				EOverrideWidgetStatus::Undetermined,
+				{
+					.Icon = FAppStyle::GetBrush(TEXT("DetailsView.OverrideUndetermined")),
+					.HoveredIcon = FAppStyle::GetBrush(TEXT("DetailsView.OverrideUndetermined")),
+					.Tooltip = LOCTEXT("UndeterminedPropertyToolTip", "State has not yet been determined"),
+					.IconSize = FVector2D(16, 16),
+				},
+			},
+			{ // uninitialized
+				EOverrideWidgetStatus::Uninitialized,
+				{
+					.Icon = FAppStyle::GetBrush(TEXT("DetailsView.OverrideAlert")),
+					.HoveredIcon = FAppStyle::GetBrush(TEXT("DetailsView.OverrideAlert.Hovered")),
+					.Tooltip = LOCTEXT("NonConcreteUninitialized", "This property needs a value.\nYou won't be able to test or publish until you set one."),
+					.Color = FStyleColors::Error,
+					.IconSize = FVector2D(16, 16),
+				}
+			},
+			{ // Inherited
+				EOverrideWidgetStatus::Inherited,
+				{
+					.Icon = FAppStyle::GetBrush(TEXT("DetailsView.OverrideInherited")),
+					.HoveredIcon = FAppStyle::GetBrush(TEXT("DetailsView.OverrideInherited.Hovered")),
+					.Tooltip = LOCTEXT("InheritedPropertyToolTip", "This property's parent component has been overridden."),
+					.HoveredColor = FColor::White,
+					.IconSize = FVector2D(16, 16),
+				}
+			},
+			{ // Here
+				EOverrideWidgetStatus::ChangedHere,
+				{
+					.Icon = FAppStyle::GetBrush(TEXT("DetailsView.OverrideHere")),
+					.HoveredIcon = FAppStyle::GetBrush(TEXT("DetailsView.OverrideHere.Hovered")),
+					.Tooltip = LOCTEXT("HerePropertyToolTip", "This property has been overridden."),
+					.HoveredColor = FColor::White,
+					.IconSize = FVector2D(16, 16),
+				}
+			},
+			{ // Inside
+				EOverrideWidgetStatus::ChangedInside,
+				{
+					.Icon = FAppStyle::GetBrush(TEXT("DetailsView.OverrideInside")),
+					.HoveredIcon = FAppStyle::GetBrush(TEXT("DetailsView.OverrideInside.Hovered")),
+					.Tooltip = LOCTEXT("InsidePropertyToolTip", "At least one of this property's values has been overridden."),
+					.HoveredColor = FColor::White,
+					.IconSize = FVector2D(16, 16),
+				}
+			},
+			{ // Outside
+				EOverrideWidgetStatus::ChangedOutside,
+				{
+					.Icon = FAppStyle::GetBrush(TEXT("DetailsView.OverrideInherited")),
+					.HoveredIcon = FAppStyle::GetBrush(TEXT("DetailsView.OverrideInherited.Hovered")),
+					.Tooltip = LOCTEXT("OutsidePropertyToolTip", "A parent of this property has been overridden."),
+					.HoveredColor = FColor::White,
+					.IconSize = FVector2D(16, 16),
+				}
+			},
+			{ // Mixed
+            	EOverrideWidgetStatus::Mixed,
+            	{
+            		.Icon = FAppStyle::GetBrush(TEXT("DetailsView.OverrideMixed")),
+            		.HoveredIcon = FAppStyle::GetBrush(TEXT("DetailsView.OverrideMixed.Hovered")),
+            		.Tooltip = LOCTEXT("OutsidePropertyToolTip", "The selected elements have mixed override states on this property."),
+					.HoveredColor = FColor::White,
+            		.IconSize = FVector2D(16, 16),
+            	}
+            }
+		};
+
+	return StyleMap.Find(InStatus);
+}
+
+#undef LOCTEXT_NAMESPACE
diff --git a/Engine/Source/Editor/AnimationEditorWidgets/Public/Overrides/OverrideStatusDetailsDisplayManager.h b/Engine/Source/Editor/AnimationEditorWidgets/Public/Overrides/OverrideStatusDetailsDisplayManager.h
new file mode 100644
--- /dev/null
+++ b/Engine/Source/Editor/AnimationEditorWidgets/Public/Overrides/OverrideStatusDetailsDisplayManager.h
@@ -0,0 +1,72 @@
+//  Copyright Epic Games, Inc. All Rights Reserved.
+
+#pragma once
+
+#include "CoreMinimal.h"
+#include "DetailsDisplayManager.h"
+#include "DetailsViewStyleKey.h"
+#include "Overrides/OverrideStatusDetailsWidgetBuilder.h"
+#include "Overrides/OverrideStatusWidgetMenuBuilder.h"
+
+class FOverrideStatusDetailsWidgetBuilder;
+
+/**
+ * The display manager is used to determine how the details view should behave when using an object
+ * filter. In this case the display manager is used to set the property-updated-widget (the override status widget).
+ */
+class ANIMATIONEDITORWIDGETS_API FOverrideStatusDetailsDisplayManager : public FDetailsDisplayManager
+{
+public:
+
+	virtual ~FOverrideStatusDetailsDisplayManager() override;
+
+	/**
+	 * Returns a @code bool @endcode indicating whether this @code DetailsViewObjectFilter @endcode instance
+	 * has a category menu
+	 */
+	virtual bool ShouldShowCategoryMenu() override;
+
+	/**
+	 * Returns the @code const FDetailsViewStyleKey& @endcode that is the Key to the current FDetailsViewStyle style
+	 */
+	virtual const FDetailsViewStyleKey& GetDetailsViewStyleKey() const override;
+
+	// Returns true if this manager can construct the property updated widget
+	virtual bool CanConstructPropertyUpdatedWidgetBuilder() const override;
+
+	// Returns the builder used to construct the property updated widgets, in this case the SOverrideStatusWidget 
+	virtual TSharedPtr<FPropertyUpdatedWidgetBuilder> ConstructPropertyUpdatedWidgetBuilder(const FConstructPropertyUpdatedWidgetBuilderArgs& Args) override;
+
+	// Returns the preconfigured menu builder for this display manager and a given subject
+	TSharedPtr<FOverrideStatusWidgetMenuBuilder> GetMenuBuilder(const FOverrideStatusSubject& InSubject) const;
+
+	FOverrideStatus_CanCreateWidget& OnCanCreateWidget() { return CanCreateWidgetDelegate; }
+	FOverrideStatus_GetStatus& OnGetStatus() { return GetStatusDelegate; }
+	FOverrideStatus_OnWidgetClicked& OnWidgetClicked() { return WidgetClickedDelegate; }
+	FOverrideStatus_OnGetMenuContent& OnGetMenuContent() { return GetMenuContentDelegate; }
+	FOverrideStatus_AddOverride& OnAddOverride() { return AddOverrideDelegate; }
+	FOverrideStatus_ClearOverride& OnClearOverride() { return ClearOverrideDelegate; }
+	FOverrideStatus_ResetToDefault& OnResetToDefault() { return ResetToDefaultDelegate; }
+	FOverrideStatus_ValueDiffersFromDefault& OnValueDiffersFromDefault() { return ValueDiffersFromDefaultDelegate; }
+
+private:
+	
+	TSharedPtr<FOverrideStatusDetailsWidgetBuilder> ConstructOverrideWidgetBuilder(const FConstructPropertyUpdatedWidgetBuilderArgs& Args);
+	void OnObjectPropertyChanged(UObject* ObjectBeingModified, FPropertyChangedEvent& PropertyChangedEvent);
+	void SetIsDisplayingOverrideableObject(bool InIsDisplayingOverrideableObject);
+
+	bool bIsDisplayingOverrideableObject = false;
+	
+	FExecuteAction InvalidateCachedState;
+
+	FOverrideStatus_CanCreateWidget CanCreateWidgetDelegate;
+	FOverrideStatus_GetStatus GetStatusDelegate;
+	FOverrideStatus_OnWidgetClicked WidgetClickedDelegate;
+	FOverrideStatus_OnGetMenuContent GetMenuContentDelegate;
+	FOverrideStatus_AddOverride AddOverrideDelegate;
+	FOverrideStatus_ClearOverride ClearOverrideDelegate;
+	FOverrideStatus_ResetToDefault ResetToDefaultDelegate;
+	FOverrideStatus_ValueDiffersFromDefault ValueDiffersFromDefaultDelegate;
+
+	friend class FOverrideStatusDetailsViewObjectFilter;
+};
diff --git a/Engine/Source/Editor/AnimationEditorWidgets/Public/Overrides/OverrideStatusDetailsObjectFilter.h b/Engine/Source/Editor/AnimationEditorWidgets/Public/Overrides/OverrideStatusDetailsObjectFilter.h
new file mode 100644
--- /dev/null
+++ b/Engine/Source/Editor/AnimationEditorWidgets/Public/Overrides/OverrideStatusDetailsObjectFilter.h
@@ -0,0 +1,64 @@
+// Copyright Epic Games, Inc. All Rights Reserved.
+
+#pragma once
+
+#include "Overrides/OverrideStatusDetailsDisplayManager.h"
+#include "Overrides/OverrideStatusWidgetMenuBuilder.h"
+#include "DetailsViewObjectFilter.h"
+
+DECLARE_DELEGATE_RetVal_TwoParams(bool, FOverrideStatusObjectFilter_CanMergeObjects, const UObject*, const UObject*);
+
+/**
+ * An object filter for the property editor / details view.
+ * The filter can decide if it can display a certain object - in this case the object filter
+ * is used to provide the override status widget instead of the reset value arrow to the details
+ * panel.
+ */
+class ANIMATIONEDITORWIDGETS_API FOverrideStatusDetailsViewObjectFilter : public FDetailsViewObjectFilter
+{	
+public:
+	// the standard method to create an object filter. 
+	template<typename T = FOverrideStatusDetailsViewObjectFilter>
+	static TSharedPtr<T> Create()
+	{
+		TSharedPtr<T> ObjectFilter = MakeShared<T>();
+		ObjectFilter->InitializeDisplayManager();
+		return ObjectFilter;
+	}
+
+	// default constructor
+	FOverrideStatusDetailsViewObjectFilter();
+
+	// sets up the  display manager for this filter
+	virtual void InitializeDisplayManager();
+
+	// Given a const TArray<UObject*>& SourceObjects, filters the objects and puts the objects which should
+	// be shown in the Details panel in the return TArray<FDetailsViewObjectRoot>. These may be some part of the
+	// original SourceObjects array, itself, or it may be some contained sub-objects within the SourceObjects.
+	virtual TArray<FDetailsViewObjectRoot> FilterObjects(const TArray<UObject*>& SourceObjects) override;
+
+	// Returns a preconfigured menu builder for this filter.
+	TSharedPtr<FOverrideStatusWidgetMenuBuilder> GetMenuBuilder(const FOverrideStatusSubject& InSubject) const;
+	
+	FOverrideStatus_CanCreateWidget& OnCanCreateWidget();
+	FOverrideStatus_GetStatus& OnGetStatus();
+	FOverrideStatus_OnWidgetClicked& OnWidgetClicked();
+	FOverrideStatus_OnGetMenuContent& OnGetMenuContent();
+	FOverrideStatus_AddOverride& OnAddOverride();
+	FOverrideStatus_ClearOverride& OnClearOverride();
+	FOverrideStatus_ResetToDefault& OnResetToDefault();
+	FOverrideStatus_ValueDiffersFromDefault& OnValueDiffersFromDefault();
+	FOverrideStatusObjectFilter_CanMergeObjects& OnCanMergeObjects() { return CanMergeObjectDelegate; }
+
+	static bool MergeObjectByClass(const UObject* InObjectA, const UObject* InObjectB);
+
+private:
+	
+	/**
+	 * The @code FOverrideStatusDetailsDisplayManager @endcode which provides an API to manage some of the characteristics of the
+	 * details display
+	 */
+	TSharedPtr<FOverrideStatusDetailsDisplayManager> OverrideStatusDisplayManager;
+
+	FOverrideStatusObjectFilter_CanMergeObjects CanMergeObjectDelegate;
+};
\ No newline at end of file
diff --git a/Engine/Source/Editor/AnimationEditorWidgets/Public/Overrides/OverrideStatusDetailsWidgetBuilder.h b/Engine/Source/Editor/AnimationEditorWidgets/Public/Overrides/OverrideStatusDetailsWidgetBuilder.h
new file mode 100644
--- /dev/null
+++ b/Engine/Source/Editor/AnimationEditorWidgets/Public/Overrides/OverrideStatusDetailsWidgetBuilder.h
@@ -0,0 +1,62 @@
+//  Copyright Epic Games, Inc. All Rights Reserved.
+
+#pragma once
+
+#include "Templates/SharedPointerFwd.h"
+#include "UserInterface/Widgets/PropertyUpdatedWidgetBuilder.h"
+#include "OverrideStatusSubject.h"
+
+class FOverrideStatusDetailsDisplayManager;
+
+/**
+ * A Display builder for the override status combo button
+ */
+class ANIMATIONEDITORWIDGETS_API FOverrideStatusDetailsWidgetBuilder : public FPropertyUpdatedWidgetBuilder
+{
+public:
+	DECLARE_DELEGATE_RetVal(TSharedRef<SWidget>, FOnGetWidget);
+
+public:
+
+	/**
+	 * The constructor, which takes a @code TSharedRef<FDetailsDisplayManager> @endcode to initialize
+	 * the Details Display Manager
+	 *
+	 * @param InDetailsDisplayManager the FDetailsDisplayManager which manages the details display
+	 * than a property row   
+	 */
+	FOverrideStatusDetailsWidgetBuilder(
+		const TSharedRef<FOverrideStatusDetailsDisplayManager>& InDetailsDisplayManager,
+		const TArray<TWeakObjectPtr<const UObject>>& InObjects,
+		const TSharedPtr<const FPropertyPath>& InPropertyPath,
+		const FName& InCategory);
+
+	/**
+	 * Implements the generation of the Category Menu button SWidget
+	 */
+	virtual TSharedPtr<SWidget> GenerateWidget() override;
+
+	/**
+	 * Converts this into the SWidget it builds
+	 */
+	virtual ~FOverrideStatusDetailsWidgetBuilder() override;
+
+	FOverrideStatus_CanCreateWidget& OnCanCreateWidget();
+	FOverrideStatus_GetStatus& OnGetStatus();
+	FOverrideStatus_OnWidgetClicked& OnWidgetClicked();
+	FOverrideStatus_OnGetMenuContent& OnGetMenuContent();
+	FOverrideStatus_AddOverride& OnAddOverride();
+	FOverrideStatus_ClearOverride& OnClearOverride();
+	FOverrideStatus_ResetToDefault& OnResetToDefault();
+	FOverrideStatus_ValueDiffersFromDefault& OnValueDiffersFromDefault();
+
+private:
+
+	/**
+	 * The @code DetailsDisplayManager @endcode which provides an API to manage some of the characteristics of the
+	 * details display
+	 */
+	TSharedRef<FOverrideStatusDetailsDisplayManager> DisplayManager;
+
+	FOverrideStatusSubject Subject;
+};
\ No newline at end of file
diff --git a/Engine/Source/Editor/AnimationEditorWidgets/Public/Overrides/OverrideStatusSubject.h b/Engine/Source/Editor/AnimationEditorWidgets/Public/Overrides/OverrideStatusSubject.h
new file mode 100644
--- /dev/null
+++ b/Engine/Source/Editor/AnimationEditorWidgets/Public/Overrides/OverrideStatusSubject.h
@@ -0,0 +1,195 @@
+//  Copyright Epic Games, Inc. All Rights Reserved.
+
+#pragma once
+
+#include "Templates/SharedPointerFwd.h"
+#include "SOverrideStatusWidget.h"
+
+class FPropertyPath;
+struct FOverrideStatusSubject;
+
+DECLARE_DELEGATE_RetVal_OneParam(bool, FOverrideStatus_CanCreateWidget, const FOverrideStatusSubject&);
+DECLARE_DELEGATE_RetVal_OneParam(EOverrideWidgetStatus::Type, FOverrideStatus_GetStatus, const FOverrideStatusSubject&);
+DECLARE_DELEGATE_RetVal_TwoParams(FReply, FOverrideStatus_OnWidgetClicked, const FOverrideStatusSubject&, EOverrideWidgetStatus::Type);
+DECLARE_DELEGATE_RetVal_TwoParams(TSharedRef<SWidget>, FOverrideStatus_OnGetMenuContent, const FOverrideStatusSubject&, EOverrideWidgetStatus::Type);
+DECLARE_DELEGATE_RetVal_OneParam(FReply, FOverrideStatus_AddOverride, const FOverrideStatusSubject&);
+DECLARE_DELEGATE_RetVal_OneParam(FReply, FOverrideStatus_ClearOverride, const FOverrideStatusSubject&);
+DECLARE_DELEGATE_RetVal_OneParam(FReply, FOverrideStatus_ResetToDefault, const FOverrideStatusSubject&);
+DECLARE_DELEGATE_RetVal_OneParam(bool, FOverrideStatus_ValueDiffersFromDefault, const FOverrideStatusSubject&);
+
+/**
+ * The subject of an override status (widget).
+ * To support multi selection the subject is represented by an array of objects and a property path.
+ * This class also offers helper template functions to facilitate the interaction between the list
+ * of subject object and the user interface layer.
+ */
+struct ANIMATIONEDITORWIDGETS_API FOverrideStatusSubject
+{
+public:
+	FOverrideStatusSubject(const TArray<TWeakObjectPtr<const UObject>>& InObjects, const TSharedPtr<const FPropertyPath>& InPropertyPath, const FName& InCategory);
+	FOverrideStatusSubject(const TWeakObjectPtr<const UObject>& InObject, const TSharedPtr<const FPropertyPath>& InPropertyPath = nullptr, const FName& InCategory = NAME_None);
+	FOverrideStatusSubject(const UObject* InObject, const TSharedPtr<const FPropertyPath>& InPropertyPath = nullptr, const FName& InCategory = NAME_None);
+
+	// returns true if any of the objects within the subject is still valid
+	bool IsValid() const;
+
+	// returns the number of objects
+	int32 Num() const;
+
+	// array access operator 
+	const UObject* operator [](int32 InIndex) const;
+
+	// range access
+	TArray<TWeakObjectPtr<const UObject>>::RangedForConstIteratorType begin() const { return Objects.begin(); }
+	TArray<TWeakObjectPtr<const UObject>>::RangedForConstIteratorType end() const { return Objects.end(); }
+
+	// returns true if this subject contains a property path
+	bool HasPropertyPath() const;
+
+	// getter for the property path
+	const TSharedPtr<const FPropertyPath>& GetPropertyPath() const;
+
+	// getter for the propery path string (uses a cache inside and is faster than GetPropertyPath()->ToString())
+	const FString& GetPropertyPathString(const TCHAR* Separator = TEXT("->")) const;
+
+	// returns true if this subject is linked to a category
+	bool HasCategory() const;
+
+	// getter for category
+	const FName& GetCategory() const;
+
+	// returns true if a given object is part of the subject
+	bool Contains(const UObject* InObject) const;
+
+	// returns true if the subject contains an object of type ObjectType
+	template<typename ObjectType>
+	bool Contains() const
+	{
+		for(int32 Index = 0; Index < Num(); Index++)
+		{
+			if(Cast<ObjectType>(operator[](Index)))
+			{
+				return true;
+			}
+		}
+		return false;
+	}
+
+	// returns true if the subject contains an object matching the given predicate
+	template<typename ObjectType = UObject>
+	bool Contains(TFunction<bool(const ObjectType*)> InMatchPredicate) const
+	{
+		for(int32 Index = 0; Index < Num(); Index++)
+		{
+			if(const ObjectType* CastObject = Cast<ObjectType>(operator[](Index)))
+			{
+				if(InMatchPredicate(CastObject))
+				{
+					return true;
+				}
+			}
+		}
+		return false;
+	}
+
+	// returns the index of a given object or INDEX_NONE
+	int32 Find(const UObject* InObject) const;
+
+	// returns the index of a given object or INDEX_NONE,
+	// the match is determined based on the given predicate 
+	template<typename ObjectType = UObject>
+	int32 Find(TFunction<bool(const ObjectType*)> InMatchPredicate) const
+	{
+		for(int32 Index = 0; Index < Num(); Index++)
+		{
+			if(const ObjectType* CastObject = Cast<ObjectType>(operator[](Index)))
+			{
+				if(InMatchPredicate(CastObject))
+				{
+					return true;
+				}
+			}
+		}
+		return false;
+	}
+
+	// executes a given predicate for each object in the subject
+	template<typename ObjectType = UObject>
+	void ForEach(TFunction<void(const ObjectType*)> InPerObjectPredicate) const
+	{
+		for(int32 Index = 0; Index < Num(); Index++)
+		{
+			if(const ObjectType* CastObject = Cast<ObjectType>(operator[](Index)))
+			{
+				InPerObjectPredicate(CastObject);
+			}
+		}
+	}
+
+	// returns the first valid value for a given predicate (or an unset TOptional). each object in the subject
+	// is interrogated for using the predicate and can return a TOptional indicating success.
+	template<typename ValueType, typename ObjectType = UObject>
+	TOptional<ValueType> GetFirstValue(TFunction<TOptional<ValueType>(const ObjectType*)> InGetValuePerObjectPredicate) const
+	{
+		for(int32 Index = 0; Index < Num(); Index++)
+		{
+			if(const ObjectType* CastObject = Cast<ObjectType>(operator[](Index)))
+			{
+				const TOptional<ValueType> SingleValue = InGetValuePerObjectPredicate(CastObject);
+				if(SingleValue.IsSet())
+				{
+					return SingleValue;
+				}
+			}
+		}
+		return TOptional<ValueType>();
+	}
+
+	// returns a common for a given predicate (or an unset TOptional). each object in the subject
+	// is interrogated for using the predicate and can return a TOptional indicating success.
+	// if the value matches for all objects the value will be returned as a set TOptional, otherwise
+	// an unset TOptional will be returned.
+	template<typename ValueType, typename ObjectType = UObject>
+	TOptional<ValueType> GetCommonValue(TFunction<TOptional<ValueType>(const ObjectType*)> InGetValuePerObjectPredicate) const
+	{
+		TOptional<ValueType> Result;
+		for(int32 Index = 0; Index < Num(); Index++)
+		{
+			if(const ObjectType* CastObject = Cast<ObjectType>(operator[](Index)))
+			{
+				const TOptional<ValueType> SingleValue = InGetValuePerObjectPredicate(CastObject);
+				if(SingleValue.IsSet())
+				{
+					if(Result.IsSet())
+					{
+						if(Result.GetValue() != SingleValue)
+						{
+							Result.Reset();
+							break;
+						}
+					}
+					else
+					{
+						Result = SingleValue;
+					}
+				}
+			}
+		}
+		return Result;
+	}
+
+	// returns the status provided by the subject's objects or an empty TOptional if it varies across objects
+	template<typename ObjectType = UObject>
+	TOptional<EOverrideWidgetStatus::Type> GetStatus(TFunction<TOptional<EOverrideWidgetStatus::Type>(const ObjectType*)> InGetValuePerObjectPredicate) const
+	{
+		return GetCommonValue<EOverrideWidgetStatus::Type, ObjectType>(InGetValuePerObjectPredicate);
+	}
+
+private:
+	TArray<TWeakObjectPtr<const UObject>> Objects;
+	TSharedPtr<const FPropertyPath> PropertyPath;
+	FName Category;
+	
+	mutable TOptional<FString> LastSeparator;
+	mutable TOptional<FString> LastPropertyPathString;
+};
diff --git a/Engine/Source/Editor/AnimationEditorWidgets/Public/Overrides/OverrideStatusWidgetMenuBuilder.h b/Engine/Source/Editor/AnimationEditorWidgets/Public/Overrides/OverrideStatusWidgetMenuBuilder.h
new file mode 100644
--- /dev/null
+++ b/Engine/Source/Editor/AnimationEditorWidgets/Public/Overrides/OverrideStatusWidgetMenuBuilder.h
@@ -0,0 +1,91 @@
+// Copyright Epic Games, Inc. All Rights Reserved.
+
+#pragma once
+#include "ToolElementRegistry.h"
+#include "ToolMenu.h"
+#include "Widgets/SWidget.h"
+#include "Overrides/OverrideStatusSubject.h"
+#include "Overrides/SOverrideStatusWidget.h"
+#include "Overrides/OverrideStatusDetailsWidgetBuilder.h"
+
+class FPropertyPath;
+class FOverrideStatusDetailsDisplayManager;
+
+/**
+ * Builder for the override status menu that goes on the top right side of each property
+ * Component in the details panel
+ */
+class ANIMATIONEDITORWIDGETS_API FOverrideStatusWidgetMenuBuilder : public FToolElementRegistrationArgs
+{
+
+	DECLARE_DELEGATE(FResetToDefault);
+	
+public:
+	/**
+	 * constructor
+	 *
+	 */
+	FOverrideStatusWidgetMenuBuilder(const FOverrideStatusSubject& InSubject, const TWeakPtr<FOverrideStatusDetailsDisplayManager>& InDisplayManager = TWeakPtr<FOverrideStatusDetailsDisplayManager>());
+
+	virtual ~FOverrideStatusWidgetMenuBuilder() override;
+
+	// returns the status of the override
+	EOverrideWidgetStatus::Type GetStatus() const;
+
+	// returns the attribute backing up the status of the override
+	TAttribute<EOverrideWidgetStatus::Type>& GetStatusAttribute() { return StatusAttribute; }
+
+	/**
+	 * Override the active overrideable object at the given property path
+	 */
+	void AddOverride();
+	bool CanAddOverride() const;
+	FOverrideStatus_AddOverride& OnAddOverride();
+	
+	/**
+	 * Clears any active overrides on the property / object
+	 */
+	void ClearOverride();
+	bool CanClearOverride() const;
+	FOverrideStatus_ClearOverride& OnClearOverride();
+
+	/**
+	 * Override the active overrideable object at the given property path
+	 */
+	void ResetToDefault();
+	bool CanResetToDefault() const;
+	FOverrideStatus_ResetToDefault& OnResetToDefault();
+	FOverrideStatus_ValueDiffersFromDefault& OnValueDiffersFromDefault();
+
+	/**
+	 * Set up the menu
+	 */
+	void InitializeMenu();
+
+	/**
+	 * Fill this method in with your Slate to create it
+	 */
+	virtual TSharedPtr<SWidget> GenerateWidget() override;
+
+private:
+	/**
+	 * The UToolMenu providing the context menu
+	 */
+	TWeakObjectPtr<UToolMenu> ToolMenu;
+
+	/**
+	 * The object that will be queried for its override state
+	 */
+	FOverrideStatusSubject Subject;
+
+	/**
+	 * The status wrapped as a property
+	 */
+	TAttribute<EOverrideWidgetStatus::Type> StatusAttribute;
+
+	FOverrideStatus_AddOverride AddOverrideDelegate;
+	FOverrideStatus_ClearOverride ClearOverrideDelegate;
+	FOverrideStatus_ResetToDefault ResetToDefaultDelegate;
+	FOverrideStatus_ValueDiffersFromDefault ValueDiffersFromDefaultDelegate;
+};
+
diff --git a/Engine/Source/Editor/AnimationEditorWidgets/Public/Overrides/SOverrideListWidget.h b/Engine/Source/Editor/AnimationEditorWidgets/Public/Overrides/SOverrideListWidget.h
new file mode 100644
--- /dev/null
+++ b/Engine/Source/Editor/AnimationEditorWidgets/Public/Overrides/SOverrideListWidget.h
@@ -0,0 +1,37 @@
+//  Copyright Epic Games, Inc. All Rights Reserved.
+
+#pragma once
+#include "Overrides/SOverrideStatusWidget.h"
+#include "Overrides/OverrideStatusSubject.h"
+#include "Widgets/Text/STextBlock.h"
+
+class ANIMATIONEDITORWIDGETS_API SOverrideListWidget : public SCompoundWidget
+{
+	SLATE_DECLARE_WIDGET(SOverrideListWidget, SCompoundWidget)
+public:
+	SLATE_BEGIN_ARGS( SOverrideListWidget )
+		: _SubjectsHash(0)
+	{}
+		SLATE_ATTRIBUTE(uint32, SubjectsHash)
+		SLATE_ATTRIBUTE(TArray<FOverrideStatusSubject>, Subjects)
+		SLATE_EVENT(FOverrideStatus_GetStatus, OnGetStatus)
+        SLATE_EVENT(FOverrideStatus_ClearOverride, OnClearOverride)
+	SLATE_END_ARGS()
+public:
+	SOverrideListWidget();
+	virtual ~SOverrideListWidget() override;
+
+	void Construct(const FArguments& InArgs);
+
+	virtual void Tick(const FGeometry& AllottedGeometry, const double InCurrentTime, const float InDeltaTime) override;
+	
+private:
+
+	TAttribute<uint32> SubjectsHashAttribute;
+	TAttribute<TArray<FOverrideStatusSubject>> SubjectsAttribute;
+	FOverrideStatus_GetStatus GetStatusDelegate;
+	FOverrideStatus_ClearOverride ClearOverrideDelegate;
+
+	TSharedPtr<STextBlock> TextBlock;
+	TOptional<uint32> LastHash;
+};
diff --git a/Engine/Source/Editor/AnimationEditorWidgets/Public/Overrides/SOverrideStatusWidget.h b/Engine/Source/Editor/AnimationEditorWidgets/Public/Overrides/SOverrideStatusWidget.h
new file mode 100644
--- /dev/null
+++ b/Engine/Source/Editor/AnimationEditorWidgets/Public/Overrides/SOverrideStatusWidget.h
@@ -0,0 +1,96 @@
+//  Copyright Epic Games, Inc. All Rights Reserved.
+
+#pragma once
+#include "Framework/SlateDelegates.h"
+#include "Widgets/SCompoundWidget.h"
+#include "Styling/StyleColors.h"
+#include "UObject/OverriddenPropertySet.h"
+
+class SImage;
+
+namespace EOverrideWidgetStatus
+{
+	enum Type : int
+	{
+		None, // No override.
+		Undetermined, // State has not yet been determine
+		Uninitialized, // This property needs a value.\nYou won't be able to test or publish until you set one. 
+		Inherited, // This property's parent component has been overridden.
+		ChangedHere, // This property has been overridden.
+		ChangedInside, // At least one of this property's values has been overridden.
+		ChangedOutside, // A parent property of this property has been overridden.
+		Mixed, // The selected elements have mixed override states on this property.
+	};
+}
+
+struct ANIMATIONEDITORWIDGETS_API FOverrideStatusWidgetStyle
+{
+	const FSlateBrush* Icon = nullptr;
+	const FSlateBrush* HoveredIcon = nullptr;
+	FText Tooltip;
+	FSlateColor Color = FStyleColors::White;
+	FSlateColor HoveredColor = FStyleColors::White;
+	TOptional<FVector2D> IconSize;
+};
+
+DECLARE_DELEGATE_RetVal_OneParam(const FOverrideStatusWidgetStyle*, FOverrideStatusWidgetGetStyle, EOverrideWidgetStatus::Type);
+DECLARE_DELEGATE_RetVal_OneParam(FText, FOverrideStatusWidgetGetTooltip, EOverrideWidgetStatus::Type);
+
+class ANIMATIONEDITORWIDGETS_API SOverrideStatusWidget : public SCompoundWidget
+{
+	SLATE_DECLARE_WIDGET(SOverrideStatusWidget, SCompoundWidget)
+public:
+	SLATE_BEGIN_ARGS( SOverrideStatusWidget )
+	: _Status(EOverrideWidgetStatus::None)
+	, _IsHovered(false)
+	{}
+		SLATE_ATTRIBUTE(EOverrideWidgetStatus::Type, Status)
+		SLATE_ARGUMENT(TOptional<FOverrideStatusWidgetStyle>, DefaultStyle)
+		SLATE_EVENT(FOverrideStatusWidgetGetStyle, Style)
+		SLATE_EVENT(FOverrideStatusWidgetGetTooltip, Tooltip)
+		SLATE_EVENT(FOnClicked, OnClicked)
+		SLATE_EVENT(FOnGetContent, MenuContent)
+		SLATE_ATTRIBUTE(bool, IsHovered)
+	SLATE_END_ARGS()
+public:
+	SOverrideStatusWidget();
+	virtual ~SOverrideStatusWidget() override;
+
+	void Construct(const FArguments& InArgs);
+
+	EOverrideWidgetStatus::Type GetStatus() const;
+	void SetStatus(EOverrideWidgetStatus::Type InStatus);
+
+	virtual FReply OnMouseButtonDown(const FGeometry& MyGeometry, const FPointerEvent& MouseEvent) override;
+	virtual FReply OnMouseButtonUp(const FGeometry& MyGeometry, const FPointerEvent& MouseEvent) override;
+	virtual void OnMouseEnter(const FGeometry& MyGeometry, const FPointerEvent& MouseEvent) override;
+	virtual void OnMouseLeave(const FPointerEvent& MouseEvent) override;
+	
+	void QueueRepaint();
+
+	virtual void Tick(const FGeometry& AllottedGeometry, const double InCurrentTime, const float InDeltaTime) override;
+
+	static const FOverrideStatusWidgetStyle* GetStyleFromStatus(EOverrideWidgetStatus::Type InStatus);
+
+private:
+	
+	FReply HandleClick();
+	static const FOverrideStatusWidgetStyle* GetStyleFromStatus_Fallback(EOverrideWidgetStatus::Type InStatus);
+	
+	/**
+	 * Style used for the button.
+	 * Note that it is never null but does change
+	 */
+	const FOverrideStatusWidgetStyle* WidgetStyle;
+	FOverrideStatusWidgetStyle DefaultWidgetStyle;
+	TSharedPtr<SImage> Image;
+	bool bAlwaysUpdateOnTick;
+	bool bTickIsPending;
+	TOptional<bool> bLastHovered;
+	
+	FOnGetContent OnGetMenuContent;
+	FOnClicked OnClicked;
+	TAttribute<EOverrideWidgetStatus::Type> StatusAttribute;
+	FOverrideStatusWidgetGetStyle OnGetStyle;
+	FOverrideStatusWidgetGetTooltip OnGetTooltip;
+};
diff --git a/Engine/Source/Editor/EditorStyle/Private/StarshipStyle.cpp b/Engine/Source/Editor/EditorStyle/Private/StarshipStyle.cpp
--- a/Engine/Source/Editor/EditorStyle/Private/StarshipStyle.cpp
+++ b/Engine/Source/Editor/EditorStyle/Private/StarshipStyle.cpp
@@ -3600,6 +3600,7 @@ void FStarshipEditorStyle::FStyle::SetupPropertyEditorStyles()
 		Set("DetailsView.OverrideHereInside", new IMAGE_BRUSH_SVG("Starship/DetailsView/DetailsOverrideHereInside", Icon16x16, FStyleColors::AccentBlue));
 		Set("DetailsView.OverrideInherited", new IMAGE_BRUSH_SVG("Starship/DetailsView/DetailsOverrideInherited", Icon16x16, FStyleColors::AccentBlue));
 		Set("DetailsView.OverrideAlert", new IMAGE_BRUSH_SVG("Starship/DetailsView/DetailsOverrideAlert", Icon16x16, FStyleColors::White));
+		Set("DetailsView.OverrideMixed", new IMAGE_BRUSH_SVG("Starship/DetailsView/DetailsOverrideMixed", Icon16x16, FStyleColors::AccentBlue));
 
 		// Hovered images for all the details panel override states
 		Set("DetailsView.OverrideHere.Hovered", new IMAGE_BRUSH_SVG("Starship/DetailsView/DetailsOverrideHere", Icon16x16, FStyleColors::White));
@@ -3610,6 +3611,7 @@ void FStarshipEditorStyle::FStyle::SetupPropertyEditorStyles()
 		Set("DetailsView.OverrideHereInside.Hovered", new IMAGE_BRUSH_SVG("Starship/DetailsView/DetailsOverrideHereInside", Icon16x16, FStyleColors::White));
 		Set("DetailsView.OverrideInherited.Hovered", new IMAGE_BRUSH_SVG("Starship/DetailsView/DetailsOverrideInherited", Icon16x16, FStyleColors::White));
 		Set("DetailsView.OverrideAlert.Hovered", new IMAGE_BRUSH_SVG("Starship/DetailsView/DetailsOverrideAlert", Icon16x16, FStyleColors::White));
+		Set("DetailsView.OverrideMixed.Hovered", new IMAGE_BRUSH_SVG("Starship/DetailsView/DetailsOverrideMixed", Icon16x16, FStyleColors::White));
 
 		
 		const FButtonStyle CategoryRowButton = FButtonStyle()
diff --git a/Engine/Source/Editor/Kismet/Private/SKismetInspector.cpp b/Engine/Source/Editor/Kismet/Private/SKismetInspector.cpp
--- a/Engine/Source/Editor/Kismet/Private/SKismetInspector.cpp
+++ b/Engine/Source/Editor/Kismet/Private/SKismetInspector.cpp
@@ -547,6 +547,11 @@ void SKismetInspector::ShowDetailsForSingleObject(UObject* Object, const FShowDe
 		PropertyObjects.Add(Object);
 	}
 
+	if(PropertyView)
+	{
+		PropertyView->SetObjectFilter(nullptr);
+	}
+
 	ShowDetailsForObjects(PropertyObjects, Options);
 }
 
diff --git a/Engine/Source/Editor/PropertyEditor/Private/SDetailCategoryTableRow.cpp b/Engine/Source/Editor/PropertyEditor/Private/SDetailCategoryTableRow.cpp
--- a/Engine/Source/Editor/PropertyEditor/Private/SDetailCategoryTableRow.cpp
+++ b/Engine/Source/Editor/PropertyEditor/Private/SDetailCategoryTableRow.cpp
@@ -132,6 +132,7 @@ void SDetailCategoryTableRow::Construct(const FArguments& InArgs, TSharedRef<FDe
 			}
 		});
 		Args.PropertyPath = nullptr;
+		Args.Category = InOwnerTreeNode->GetNodeName();
 		Args.Objects = InArgs._Objects;
 		
 		PropertyUpdatedWidgetBuilder = DisplayManager->ConstructPropertyUpdatedWidgetBuilder(Args);
diff --git a/Engine/Source/Editor/PropertyEditor/Private/SDetailSingleItemRow.cpp b/Engine/Source/Editor/PropertyEditor/Private/SDetailSingleItemRow.cpp
--- a/Engine/Source/Editor/PropertyEditor/Private/SDetailSingleItemRow.cpp
+++ b/Engine/Source/Editor/PropertyEditor/Private/SDetailSingleItemRow.cpp
@@ -648,6 +648,7 @@ void SDetailSingleItemRow::Construct( const FArguments& InArgs, FDetailLayoutCus
 							Args.ResetToDefaultAction = FExecuteAction::CreateSP(this, &SDetailSingleItemRow::OnResetToDefaultClicked);
 							Args.Objects = &Objects;
 							Args.PropertyPath = PropertyPath;
+							Args.Category = Category.IsValid() ? Category->GetCategoryName() : NAME_None;
 						
 							PropertyUpdatedWidgetBuilder = DisplayManager->ConstructPropertyUpdatedWidgetBuilder(Args);
 
@@ -666,6 +667,13 @@ void SDetailSingleItemRow::Construct( const FArguments& InArgs, FDetailLayoutCus
 			{
 				RightmostWidget = PropertyUpdatedWidgetBuilder->GenerateWidget();
 			}
+
+			// also perform the two additional cases below in case
+			// PropertyUpdatedWidgetBuilder->GenerateWidget() return nullptr
+			if(RightmostWidget.IsValid())
+			{
+				// nothing to do
+			}
 			else if(WidgetRow.HasResetToDefaultContent())
 			{
 				RightmostWidget = WidgetRow.ResetToDefaultWidget.Widget;
diff --git a/Engine/Source/Editor/PropertyEditor/Public/DetailsDisplayManager.h b/Engine/Source/Editor/PropertyEditor/Public/DetailsDisplayManager.h
--- a/Engine/Source/Editor/PropertyEditor/Public/DetailsDisplayManager.h
+++ b/Engine/Source/Editor/PropertyEditor/Public/DetailsDisplayManager.h
@@ -34,7 +34,7 @@ public:
 	/**
 	 * Returns a boolean indicating if the Component Editor should be hidden 
 	 */
-	virtual bool ShouldHideComponentEditor();
+	PROPERTYEDITOR_API virtual bool ShouldHideComponentEditor();
 
 	/**
 	 * Returns a boolean indicating whether the sub-object editor should show regardless of
@@ -46,14 +46,14 @@ public:
 	 * Returns a @code bool @endcode indicating whether this @code DetailsViewObjectFilter @endcode instance
 	 * should show a category menu
 	 */
-	virtual bool ShouldShowCategoryMenu();
+	PROPERTYEDITOR_API virtual bool ShouldShowCategoryMenu();
 
 	/**
 	 * Gets the category menu SWidget and returns a shared pointer to it
 	 *
 	 *  @param InCategoryObjectName the name of the category
 	 */
-	virtual TSharedPtr<SWidget> GetCategoryMenu(FName InCategoryObjectName);
+	PROPERTYEDITOR_API virtual TSharedPtr<SWidget> GetCategoryMenu(FName InCategoryObjectName);
 
 	/**
 	 * Updates the current details view
@@ -63,7 +63,7 @@ public:
 	/**
 	 * Returns the @code FDetailsViewStyleKey @endcode that is the Key to the current FDetailsViewStyle style
 	 */
-	virtual const FDetailsViewStyleKey& GetDetailsViewStyleKey() const;
+	PROPERTYEDITOR_API virtual const FDetailsViewStyleKey& GetDetailsViewStyleKey() const;
 
 	/** sets whether the currently active category is an Outer category*/
 	void SetIsOuterCategory(bool bInIsOuterCategory);
@@ -161,7 +161,7 @@ public:
 	*
 	* @param InNode the @code UObject* @endcode which will be tested to see if it needs an empty category stub
 	*/
-	virtual bool ShowEmptyCategoryIfRootUObjectHasNoPropertyData(UObject* InNode) const;
+	PROPERTYEDITOR_API virtual bool ShowEmptyCategoryIfRootUObjectHasNoPropertyData(UObject* InNode) const;
 
 protected:
 
@@ -192,6 +192,7 @@ struct FConstructPropertyUpdatedWidgetBuilderArgs
 	
 	FExecuteAction ResetToDefaultAction;
 	TSharedPtr<FPropertyPath> PropertyPath;
+	FName Category = NAME_None;
 	FExecuteAction InvalidateCachedState;
 	TArray<TWeakObjectPtr<UObject>>* Objects = nullptr;
 };
\ No newline at end of file
