Fixed an issue with tooltip names in SOverrideStatusWidgets

diff --git a/Engine/Source/Editor/AnimationEditorWidgets/Private/Overrides/SOverrideStatusWidget.cpp b/Engine/Source/Editor/AnimationEditorWidgets/Private/Overrides/SOverrideStatusWidget.cpp
--- a/Engine/Source/Editor/AnimationEditorWidgets/Private/Overrides/SOverrideStatusWidget.cpp
+++ b/Engine/Source/Editor/AnimationEditorWidgets/Private/Overrides/SOverrideStatusWidget.cpp
@@ -279,7 +279,7 @@ const FOverrideStatusWidgetStyle* SOverrideStatusWidget::GetStyleFromStatus(EOve
 				{
 					.Icon = FAppStyle::GetBrush(TEXT("DetailsView.OverrideInherited")),
 					.HoveredIcon = FAppStyle::GetBrush(TEXT("DetailsView.OverrideInherited.Hovered")),
-					.Tooltip = LOCTEXT("OutsidePropertyToolTip", "A parent of this property has been overridden."),
+					.Tooltip = LOCTEXT("OverrideInheritedToolTip", "A parent of this property has been overridden."),
 					.HoveredColor = FColor::White,
 					.IconSize = FVector2D(16, 16),
 				}
@@ -289,7 +289,7 @@ const FOverrideStatusWidgetStyle* SOverrideStatusWidget::GetStyleFromStatus(EOve
             	{
             		.Icon = FAppStyle::GetBrush(TEXT("DetailsView.OverrideMixed")),
             		.HoveredIcon = FAppStyle::GetBrush(TEXT("DetailsView.OverrideMixed.Hovered")),
-            		.Tooltip = LOCTEXT("OutsidePropertyToolTip", "The selected elements have mixed override states on this property."),
+            		.Tooltip = LOCTEXT("OverrideMixedToolTip", "The selected elements have mixed override states on this property."),
 					.HoveredColor = FColor::White,
             		.IconSize = FVector2D(16, 16),
             	}
