Updated headers using UnrealCodeFixup to make sure dllstorage is on methods/staticvars instead of types. Part 4/n

diff --git a/Engine/Source/Editor/AnimationEditorWidgets/Public/SchematicGraphPanel/SSchematicGraphPanel.h b/Engine/Source/Editor/AnimationEditorWidgets/Public/SchematicGraphPanel/SSchematicGraphPanel.h
--- a/Engine/Source/Editor/AnimationEditorWidgets/Public/SchematicGraphPanel/SSchematicGraphPanel.h
+++ b/Engine/Source/Editor/AnimationEditorWidgets/Public/SchematicGraphPanel/SSchematicGraphPanel.h
@@ -11,16 +11,16 @@
 class SSchematicGraphPanel;
 class SSchematicGraphNode;
 
-class ANIMATIONEDITORWIDGETS_API FSchematicGraphNodeDragDropOp : public FDragDropOperation
+class FSchematicGraphNodeDragDropOp : public FDragDropOperation
 {
 public:
 	DRAG_DROP_OPERATOR_TYPE(FSchematicGraphNodeDragDropOp, FDragDropOperation)
 
-	static TSharedRef<FSchematicGraphNodeDragDropOp> New(TArray<SSchematicGraphNode*> InSchematicGraphNodes, const TArray<FGuid>& InElements);
+	static ANIMATIONEDITORWIDGETS_API TSharedRef<FSchematicGraphNodeDragDropOp> New(TArray<SSchematicGraphNode*> InSchematicGraphNodes, const TArray<FGuid>& InElements);
 
-	virtual ~FSchematicGraphNodeDragDropOp() override;
+	ANIMATIONEDITORWIDGETS_API virtual ~FSchematicGraphNodeDragDropOp() override;
 	
-	virtual TSharedPtr<SWidget> GetDefaultDecorator() const override;
+	ANIMATIONEDITORWIDGETS_API virtual TSharedPtr<SWidget> GetDefaultDecorator() const override;
 
 	/** @return true if this drag operation contains property paths */
 	bool HasElements() const
@@ -35,9 +35,9 @@ public:
 	}
 
 	/** @return The nodes being dragged */
-	const TArray<const FSchematicGraphNode*> GetNodes() const;
+	ANIMATIONEDITORWIDGETS_API const TArray<const FSchematicGraphNode*> GetNodes() const;
 
-	FString GetJoinedDecoratorLabels() const;
+	ANIMATIONEDITORWIDGETS_API FString GetJoinedDecoratorLabels() const;
 
 private:
 
@@ -48,7 +48,7 @@ private:
 	TArray<FGuid> Elements;
 };
 
-class ANIMATIONEDITORWIDGETS_API SSchematicGraphNode : public SNodePanel::SNode
+class SSchematicGraphNode : public SNodePanel::SNode
 {
 public:
 
@@ -76,32 +76,32 @@ public:
 	SLATE_ATTRIBUTE(FText, ToolTipText)
 	SLATE_END_ARGS()
 
-	void Construct(const FArguments& InArgs);
+	ANIMATIONEDITORWIDGETS_API void Construct(const FArguments& InArgs);
 
-	virtual FVector2D ComputeDesiredSize(float LayoutScaleMultiplier) const override;
-	virtual int32 OnPaint(const FPaintArgs& Args, const FGeometry& AllottedGeometry, const FSlateRect& MyCullingRect, FSlateWindowElementList& OutDrawElements, int32 LayerId, const FWidgetStyle& InWidgetStyle, bool bParentEnabled) const override;
+	ANIMATIONEDITORWIDGETS_API virtual FVector2D ComputeDesiredSize(float LayoutScaleMultiplier) const override;
+	ANIMATIONEDITORWIDGETS_API virtual int32 OnPaint(const FPaintArgs& Args, const FGeometry& AllottedGeometry, const FSlateRect& MyCullingRect, FSlateWindowElementList& OutDrawElements, int32 LayerId, const FWidgetStyle& InWidgetStyle, bool bParentEnabled) const override;
 	
-	virtual void OnMouseEnter(const FGeometry& MyGeometry, const FPointerEvent& MouseEvent) override;
-	virtual void OnMouseLeave(const FPointerEvent& MouseEvent) override;
-	virtual FReply OnMouseButtonDown(const FGeometry& MyGeometry, const FPointerEvent& MouseEvent) override;
-	virtual FReply OnDragOver(const FGeometry& MyGeometry, const FDragDropEvent& DragDropEvent) override;
-	virtual void OnDragLeave(const FDragDropEvent& DragDropEvent) override;
-	virtual FReply OnDrop(const FGeometry& MyGeometry, const FDragDropEvent& DragDropEvent) override;
-	virtual FReply OnDragDetected(const FGeometry& MyGeometry, const FPointerEvent& MouseEvent) override;
-
-	virtual EVisibility GetNodeVisibility() const;
-	virtual FVector2f GetPosition2f() const override;
-	void EnablePositionAnimation(bool bEnabled = true);
+	ANIMATIONEDITORWIDGETS_API virtual void OnMouseEnter(const FGeometry& MyGeometry, const FPointerEvent& MouseEvent) override;
+	ANIMATIONEDITORWIDGETS_API virtual void OnMouseLeave(const FPointerEvent& MouseEvent) override;
+	ANIMATIONEDITORWIDGETS_API virtual FReply OnMouseButtonDown(const FGeometry& MyGeometry, const FPointerEvent& MouseEvent) override;
+	ANIMATIONEDITORWIDGETS_API virtual FReply OnDragOver(const FGeometry& MyGeometry, const FDragDropEvent& DragDropEvent) override;
+	ANIMATIONEDITORWIDGETS_API virtual void OnDragLeave(const FDragDropEvent& DragDropEvent) override;
+	ANIMATIONEDITORWIDGETS_API virtual FReply OnDrop(const FGeometry& MyGeometry, const FDragDropEvent& DragDropEvent) override;
+	ANIMATIONEDITORWIDGETS_API virtual FReply OnDragDetected(const FGeometry& MyGeometry, const FPointerEvent& MouseEvent) override;
+
+	ANIMATIONEDITORWIDGETS_API virtual EVisibility GetNodeVisibility() const;
+	ANIMATIONEDITORWIDGETS_API virtual FVector2f GetPosition2f() const override;
+	ANIMATIONEDITORWIDGETS_API void EnablePositionAnimation(bool bEnabled = true);
 
 	const FVector2d& GetOriginalSize() const { return OriginalSize; }
 
 	const FSchematicGraphNode* GetNodeData() const { return NodeData.Get(); }
 	FSchematicGraphNode* GetNodeData() { return NodeData.Get(); }
-	const FGuid GetGuid() const;
-	bool IsInteractive() const;
+	ANIMATIONEDITORWIDGETS_API const FGuid GetGuid() const;
+	ANIMATIONEDITORWIDGETS_API bool IsInteractive() const;
 
 	const bool IsBeingDragged() const { return bIsBeingDragged; }
-	const bool IsFadedOut() const;
+	ANIMATIONEDITORWIDGETS_API const bool IsFadedOut() const;
 	
 private:
 
@@ -132,7 +132,7 @@ private:
 
 
 /** Widget allowing editing of a control rig's structure */
-class ANIMATIONEDITORWIDGETS_API SSchematicGraphPanel : public SNodePanel, public FTickableEditorObject
+class SSchematicGraphPanel : public SNodePanel, public FTickableEditorObject
 {
 public:
 
@@ -188,58 +188,58 @@ public:
 	UE_DEPRECATED(5.6, "GetSchematicGraph with raw pointer has been deprecated, please, use GetSchematicGraphModel with a weak pointer instead.")
 	FSchematicGraphModel* GetSchematicGraph() { return nullptr; }
 
-	void SetSchematicGraphModel(TWeakPtr<FSchematicGraphModel> InGraphData);
+	ANIMATIONEDITORWIDGETS_API void SetSchematicGraphModel(TWeakPtr<FSchematicGraphModel> InGraphData);
 	const TWeakPtr<const FSchematicGraphModel> GetSchematicGraphModel() const { return GraphDataWeak; }
 	TWeakPtr<FSchematicGraphModel> GetSchematicGraphModel() { return GraphDataWeak; }
-	void Construct(const FArguments& InArgs);
-
-	void RebuildPanel();
-	void AddNode(const FSchematicGraphNode* InNodeToAdd);
-	void RemoveNode(const FSchematicGraphNode* InNodeToRemove);
-	const SSchematicGraphNode* FindNode(const FGuid& InGuid) const;
-	SSchematicGraphNode* FindNode(const FGuid& InGuid);
-	void AddLink(const FSchematicGraphLink* InLinkToAdd);
-	void RemoveLink(const FSchematicGraphLink* InLinkToRemove);
-	const FSchematicLinkWidgetInfo* FindLink(const FGuid& InGuid) const;
-	FSchematicLinkWidgetInfo* FindLink(const FGuid& InGuid);
+	ANIMATIONEDITORWIDGETS_API void Construct(const FArguments& InArgs);
+
+	ANIMATIONEDITORWIDGETS_API void RebuildPanel();
+	ANIMATIONEDITORWIDGETS_API void AddNode(const FSchematicGraphNode* InNodeToAdd);
+	ANIMATIONEDITORWIDGETS_API void RemoveNode(const FSchematicGraphNode* InNodeToRemove);
+	ANIMATIONEDITORWIDGETS_API const SSchematicGraphNode* FindNode(const FGuid& InGuid) const;
+	ANIMATIONEDITORWIDGETS_API SSchematicGraphNode* FindNode(const FGuid& InGuid);
+	ANIMATIONEDITORWIDGETS_API void AddLink(const FSchematicGraphLink* InLinkToAdd);
+	ANIMATIONEDITORWIDGETS_API void RemoveLink(const FSchematicGraphLink* InLinkToRemove);
+	ANIMATIONEDITORWIDGETS_API const FSchematicLinkWidgetInfo* FindLink(const FGuid& InGuid) const;
+	ANIMATIONEDITORWIDGETS_API FSchematicLinkWidgetInfo* FindLink(const FGuid& InGuid);
 
 	// SNodePanel interface
-	virtual void OnArrangeChildren(const FGeometry& AllottedGeometry, FArrangedChildren& ArrangedChildren) const override;
-	virtual int32 OnPaint(const FPaintArgs& Args, const FGeometry& AllottedGeometry, const FSlateRect& MyCullingRect, FSlateWindowElementList& OutDrawElements, int32 LayerId, const FWidgetStyle& InWidgetStyle, bool bParentEnabled) const override;
-	virtual void RemoveAllNodes() override;
-	virtual FReply OnMouseWheel(const FGeometry& MyGeometry, const FPointerEvent& MouseEvent) override;
+	ANIMATIONEDITORWIDGETS_API virtual void OnArrangeChildren(const FGeometry& AllottedGeometry, FArrangedChildren& ArrangedChildren) const override;
+	ANIMATIONEDITORWIDGETS_API virtual int32 OnPaint(const FPaintArgs& Args, const FGeometry& AllottedGeometry, const FSlateRect& MyCullingRect, FSlateWindowElementList& OutDrawElements, int32 LayerId, const FWidgetStyle& InWidgetStyle, bool bParentEnabled) const override;
+	ANIMATIONEDITORWIDGETS_API virtual void RemoveAllNodes() override;
+	ANIMATIONEDITORWIDGETS_API virtual FReply OnMouseWheel(const FGeometry& MyGeometry, const FPointerEvent& MouseEvent) override;
 	// End of SNodePanel interface
 	
-	TSharedRef<SSchematicGraphNode> GetChild(int32 ChildIndex) const;
+	ANIMATIONEDITORWIDGETS_API TSharedRef<SSchematicGraphNode> GetChild(int32 ChildIndex) const;
 
 	// FTickableEditorObject Interface
-	virtual TStatId GetStatId() const override;
-	virtual void Tick(float DeltaTime) override;
-	virtual void Tick( const FGeometry& AllottedGeometry, const double InCurrentTime, const float InDeltaTime ) override;
+	ANIMATIONEDITORWIDGETS_API virtual TStatId GetStatId() const override;
+	ANIMATIONEDITORWIDGETS_API virtual void Tick(float DeltaTime) override;
+	ANIMATIONEDITORWIDGETS_API virtual void Tick( const FGeometry& AllottedGeometry, const double InCurrentTime, const float InDeltaTime ) override;
 	virtual bool IsTickable() const override { return true; }
 	// End of FTickableEditorObject interface
 
-	void ToggleVisibility();
-	void OnNodeClicked(SSchematicGraphNode* Node, const FPointerEvent& MouseEvent);
-	void OnBeginDragEvent(SSchematicGraphNode* Node, const TSharedPtr<FDragDropOperation>& InDragDropEvent);
-	void OnEndDragEvent(SSchematicGraphNode* Node, const TSharedPtr<FDragDropOperation>& InDragDropEvent);
-	void OnEnterDragEvent(const TSharedPtr<FDragDropOperation>& InDragDropEvent);
-	void OnLeaveDragEvent(const TSharedPtr<FDragDropOperation>& InDragDropEvent);
-	void OnCancelDragEvent(SSchematicGraphNode* Node, const TSharedPtr<FDragDropOperation>& InDragDropEvent);
-	void OnDropEvent(SSchematicGraphNode* Node, const FDragDropEvent& InDragDropEvent);
-	FReply HandleNodeDragDetected(FGuid Guid, const FGeometry& MyGeometry, const FPointerEvent& MouseEvent);
+	ANIMATIONEDITORWIDGETS_API void ToggleVisibility();
+	ANIMATIONEDITORWIDGETS_API void OnNodeClicked(SSchematicGraphNode* Node, const FPointerEvent& MouseEvent);
+	ANIMATIONEDITORWIDGETS_API void OnBeginDragEvent(SSchematicGraphNode* Node, const TSharedPtr<FDragDropOperation>& InDragDropEvent);
+	ANIMATIONEDITORWIDGETS_API void OnEndDragEvent(SSchematicGraphNode* Node, const TSharedPtr<FDragDropOperation>& InDragDropEvent);
+	ANIMATIONEDITORWIDGETS_API void OnEnterDragEvent(const TSharedPtr<FDragDropOperation>& InDragDropEvent);
+	ANIMATIONEDITORWIDGETS_API void OnLeaveDragEvent(const TSharedPtr<FDragDropOperation>& InDragDropEvent);
+	ANIMATIONEDITORWIDGETS_API void OnCancelDragEvent(SSchematicGraphNode* Node, const TSharedPtr<FDragDropOperation>& InDragDropEvent);
+	ANIMATIONEDITORWIDGETS_API void OnDropEvent(SSchematicGraphNode* Node, const FDragDropEvent& InDragDropEvent);
+	ANIMATIONEDITORWIDGETS_API FReply HandleNodeDragDetected(FGuid Guid, const FGeometry& MyGeometry, const FPointerEvent& MouseEvent);
 
-	virtual FVector2d GetPositionForNode(FGuid InNodeGuid) const;
-	virtual FLinearColor GetColorForNode(FGuid InNodeGuid, int32 InLayerIndex) const;
-	virtual FText GetToolTipForNode(FGuid InNodeGuid) const;
-	virtual float GetScaleForNode(FGuid InNodeGuid) const;
-	void AdjustPositionWithDPIScale(FVector2d& InOutPosition, bool bInverse = false) const;
+	ANIMATIONEDITORWIDGETS_API virtual FVector2d GetPositionForNode(FGuid InNodeGuid) const;
+	ANIMATIONEDITORWIDGETS_API virtual FLinearColor GetColorForNode(FGuid InNodeGuid, int32 InLayerIndex) const;
+	ANIMATIONEDITORWIDGETS_API virtual FText GetToolTipForNode(FGuid InNodeGuid) const;
+	ANIMATIONEDITORWIDGETS_API virtual float GetScaleForNode(FGuid InNodeGuid) const;
+	ANIMATIONEDITORWIDGETS_API void AdjustPositionWithDPIScale(FVector2d& InOutPosition, bool bInverse = false) const;
 
-	virtual bool IsAutoGroupingEnabled() const;
-	virtual float GetAutoGroupingDistance() const;
+	ANIMATIONEDITORWIDGETS_API virtual bool IsAutoGroupingEnabled() const;
+	ANIMATIONEDITORWIDGETS_API virtual float GetAutoGroupingDistance() const;
 
-	virtual bool IsAutoScaleEnabledForNode(FGuid InNodeGuid) const;
-	virtual float GetMinimumLinkDistanceForNode(FGuid InLinkGuid, bool bIncludeScale = true) const;
+	ANIMATIONEDITORWIDGETS_API virtual bool IsAutoScaleEnabledForNode(FGuid InNodeGuid) const;
+	ANIMATIONEDITORWIDGETS_API virtual float GetMinimumLinkDistanceForNode(FGuid InLinkGuid, bool bIncludeScale = true) const;
 
 	FOnNodeClicked& OnNodeClicked() { return OnNodeClickedDelegate; }
 	FOnBeginDrag& OnBeginDrag() { return OnBeginDragDelegate; }
@@ -250,13 +250,13 @@ public:
 	FOnDrop& OnAcceptDrop() { return OnDropDelegate; }
 
 	const FVector2d& GetNodeLabelOffset() const { return NodeLabelOffset; }
-	void IncrementNodeLabelOffset(const FVector2d& InOffset);
+	ANIMATIONEDITORWIDGETS_API void IncrementNodeLabelOffset(const FVector2d& InOffset);
 
 private:
 
-	void UpdatePerNodeCaches(bool bRemoveNodesFromAutoGroups);
-	void UpdateAutoGroupingForNodes();
-	void UpdateAutoScalingForNodes();
+	ANIMATIONEDITORWIDGETS_API void UpdatePerNodeCaches(bool bRemoveNodesFromAutoGroups);
+	ANIMATIONEDITORWIDGETS_API void UpdateAutoGroupingForNodes();
+	ANIMATIONEDITORWIDGETS_API void UpdateAutoScalingForNodes();
 
 	bool bIsDragDropping = false;
 	TSharedPtr<FDragDropOperation> DragDropOpFromOutside;
@@ -312,4 +312,4 @@ private:
 	FVector2d NodeLabelOffset = FVector2d::ZeroVector;
 };
 
-#endif
\ No newline at end of file
+#endif
diff --git a/Engine/Source/Editor/AnimationEditorWidgets/Public/SchematicGraphPanel/SchematicGraphLink.h b/Engine/Source/Editor/AnimationEditorWidgets/Public/SchematicGraphPanel/SchematicGraphLink.h
--- a/Engine/Source/Editor/AnimationEditorWidgets/Public/SchematicGraphPanel/SchematicGraphLink.h
+++ b/Engine/Source/Editor/AnimationEditorWidgets/Public/SchematicGraphPanel/SchematicGraphLink.h
@@ -8,7 +8,7 @@
 #define SCHEMATICGRAPHLINK_BODY(ClassName, SuperClass) \
 SCHEMATICGRAPHELEMENT_BODY(ClassName, SuperClass, FSchematicGraphNode)
 
-class ANIMATIONEDITORWIDGETS_API FSchematicGraphLink : public TSharedFromThis<FSchematicGraphLink>
+class FSchematicGraphLink : public TSharedFromThis<FSchematicGraphLink>
 {
 public:
 
@@ -52,4 +52,4 @@ protected:
 	friend class FSchematicGraphModel;
 };
 
-#endif
\ No newline at end of file
+#endif
diff --git a/Engine/Source/Editor/AnimationEditorWidgets/Public/SchematicGraphPanel/SchematicGraphModel.h b/Engine/Source/Editor/AnimationEditorWidgets/Public/SchematicGraphPanel/SchematicGraphModel.h
--- a/Engine/Source/Editor/AnimationEditorWidgets/Public/SchematicGraphPanel/SchematicGraphModel.h
+++ b/Engine/Source/Editor/AnimationEditorWidgets/Public/SchematicGraphPanel/SchematicGraphModel.h
@@ -10,13 +10,13 @@
 
 class SSchematicGraphPanel;
 
-class ANIMATIONEDITORWIDGETS_API FSchematicGraphModel : public TSharedFromThis<FSchematicGraphModel>
+class FSchematicGraphModel : public TSharedFromThis<FSchematicGraphModel>
 {
 public:
 
 	SCHEMATICGRAPHELEMENT_BODY_BASE(FSchematicGraphModel)
 	
-	virtual void Reset();
+	ANIMATIONEDITORWIDGETS_API virtual void Reset();
 	virtual void ApplyToPanel(SSchematicGraphPanel* InPanel) {};
 	
 	template<typename NodeType = FSchematicGraphNode>
@@ -69,71 +69,71 @@ public:
 		return const_cast<NodeType*>(ConstThis->FindNodeChecked<NodeType>(InNodeGuid));
 	}
 
-	virtual bool RemoveNode(const FGuid& InNodeGuid);
+	ANIMATIONEDITORWIDGETS_API virtual bool RemoveNode(const FGuid& InNodeGuid);
 
-	virtual bool SetParentNode(const FGuid& InChildNodeGuid, const FGuid& InParentNodeGuid, bool bUpdateGroupNode = true);
-	virtual bool SetParentNode(const FSchematicGraphNode* InChildNode, const FSchematicGraphNode* InParentNode, bool bUpdateGroupNode = true);
-	bool RemoveFromParentNode(const FGuid& InChildNodeGuid, bool bUpdateGroupNode = true);
-	bool RemoveFromParentNode(const FSchematicGraphNode* InChildNode, bool bUpdateGroupNode = true);
+	ANIMATIONEDITORWIDGETS_API virtual bool SetParentNode(const FGuid& InChildNodeGuid, const FGuid& InParentNodeGuid, bool bUpdateGroupNode = true);
+	ANIMATIONEDITORWIDGETS_API virtual bool SetParentNode(const FSchematicGraphNode* InChildNode, const FSchematicGraphNode* InParentNode, bool bUpdateGroupNode = true);
+	ANIMATIONEDITORWIDGETS_API bool RemoveFromParentNode(const FGuid& InChildNodeGuid, bool bUpdateGroupNode = true);
+	ANIMATIONEDITORWIDGETS_API bool RemoveFromParentNode(const FSchematicGraphNode* InChildNode, bool bUpdateGroupNode = true);
 
 	const TArray<TSharedPtr<FSchematicGraphNode>>& GetNodes() const { return Nodes; }
 
-	FVector2d GetPositionForNode(const FGuid& InNodeGuid) const;
-	virtual FVector2d GetPositionForNode(const FSchematicGraphNode* InNode) const;
-	FVector2d GetPositionOffsetForNode(const FGuid& InNodeGuid) const;
-	virtual FVector2d GetPositionOffsetForNode(const FSchematicGraphNode* InNode) const;
-	bool GetPositionAnimationEnabledForNode(const FGuid& InNodeGuid) const;
-	virtual bool GetPositionAnimationEnabledForNode(const FSchematicGraphNode* InNode) const;
-	FVector2d GetSizeForNode(const FGuid& InNodeGuid) const;
-	virtual FVector2d GetSizeForNode(const FSchematicGraphNode* InNode) const;
-	float GetScaleForNode(const FGuid& InNodeGuid) const;
-	virtual float GetScaleForNode(const FSchematicGraphNode* InNode) const;
-	float GetScaleOffsetForNode(const FGuid& InNodeGuid) const;
-	virtual float GetScaleOffsetForNode(const FSchematicGraphNode* InNode) const;
-	float GetMinimumLinkDistanceForNode(const FGuid& InNodeGuid) const;
-	virtual float GetMinimumLinkDistanceForNode(const FSchematicGraphNode* InNode) const;
-	bool IsAutoScaleEnabledForNode(const FGuid& InNodeGuid) const;
-	virtual bool IsAutoScaleEnabledForNode(const FSchematicGraphNode* InNode) const;
-	int32 GetNumLayersForNode(const FGuid& InNodeGuid) const;
-	virtual int32 GetNumLayersForNode(const FSchematicGraphNode* InNode) const;
-	FLinearColor GetColorForNode(const FGuid& InNodeGuid, int32 InLayerIndex) const;
-	virtual FLinearColor GetColorForNode(const FSchematicGraphNode* InNode, int32 InLayerIndex) const;
-	const FSlateBrush* GetBrushForNode(const FGuid& InNodeGuid, int32 InLayerIndex) const;
-	virtual const FSlateBrush* GetBrushForNode(const FSchematicGraphNode* InNode, int32 InLayerIndex) const;
-	FText GetToolTipForNode(const FGuid& InNodeGuid) const;
-	virtual FText GetToolTipForNode(const FSchematicGraphNode* InNode) const;
-	ESchematicGraphVisibility::Type GetVisibilityForNode(const FGuid& InNodeGuid) const;
-	virtual ESchematicGraphVisibility::Type GetVisibilityForNode(const FSchematicGraphNode* InNode) const;
-	TOptional<ESchematicGraphVisibility::Type> GetVisibilityForChildNode(const FGuid& InParentNodeGuid, const FGuid& InChildNodeGuid) const;
-	virtual TOptional<ESchematicGraphVisibility::Type> GetVisibilityForChildNode(const FSchematicGraphNode* InParentNode, const FSchematicGraphNode* InChildNode) const;
-	TOptional<FVector2d> GetPositionForChildNode(const FGuid& InParentNodeGuid, const FGuid& InChildNodeGuid) const;
-	virtual TOptional<FVector2d> GetPositionForChildNode(const FSchematicGraphNode* InParentNode, const FSchematicGraphNode* InChildNode) const;
-	TOptional<float> GetScaleForChildNode(const FGuid& InParentNodeGuid, const FGuid& InChildNodeGuid) const;
-	virtual TOptional<float> GetScaleForChildNode(const FSchematicGraphNode* InParentNode, const FSchematicGraphNode* InChildNode) const;
-	TOptional<bool> GetInteractivityForChildNode(const FGuid& InParentNodeGuid, const FGuid& InChildNodeGuid) const;
-	virtual TOptional<bool> GetInteractivityForChildNode(const FSchematicGraphNode* InParentNode, const FSchematicGraphNode* InChildNode) const;
-	bool IsDragSupportedForNode(const FGuid& InNodeGuid) const;
-	virtual bool IsDragSupportedForNode(const FSchematicGraphNode* InNode) const;
-	virtual bool GetContextMenuForNode(const FSchematicGraphNode* InNode, FMenuBuilder& OutMenu) const;
-	const TArray<TSharedPtr<FSchematicGraphNode>> GetSelectedNodes() const;
-	void ClearSelection();
-
-	FLinearColor GetBackgroundColorForTag(const FGuid& InNodeGuid, const FGuid& InTagGuid) const;
-	virtual FLinearColor GetBackgroundColorForTag(const FSchematicGraphTag* InTag) const;
-	FLinearColor GetForegroundColorForTag(const FGuid& InNodeGuid, const FGuid& InTagGuid) const;
-	virtual FLinearColor GetForegroundColorForTag(const FSchematicGraphTag* InTag) const;
-	FLinearColor GetLabelColorForTag(const FGuid& InNodeGuid, const FGuid& InTagGuid) const;
-	virtual FLinearColor GetLabelColorForTag(const FSchematicGraphTag* InTag) const;
-	const FSlateBrush* GetBackgroundBrushForTag(const FGuid& InNodeGuid, const FGuid& InTagGuid) const;
-	virtual const FSlateBrush* GetBackgroundBrushForTag(const FSchematicGraphTag* InTag) const;
-	const FSlateBrush* GetForegroundBrushForTag(const FGuid& InNodeGuid, const FGuid& InTagGuid) const;
-	virtual const FSlateBrush* GetForegroundBrushForTag(const FSchematicGraphTag* InTag) const;
-	const FText GetLabelForTag(const FGuid& InNodeGuid, const FGuid& InTagGuid) const;
-	virtual const FText GetLabelForTag(const FSchematicGraphTag* InTag) const;
-	const FText GetToolTipForTag(const FGuid& InNodeGuid, const FGuid& InTagGuid) const;
-	virtual const FText GetToolTipForTag(const FSchematicGraphTag* InTag) const;
-	ESchematicGraphVisibility::Type GetVisibilityForTag(const FGuid& InNodeGuid, const FGuid& InTagGuid) const;
-	virtual ESchematicGraphVisibility::Type GetVisibilityForTag(const FSchematicGraphTag* InTag) const;
+	ANIMATIONEDITORWIDGETS_API FVector2d GetPositionForNode(const FGuid& InNodeGuid) const;
+	ANIMATIONEDITORWIDGETS_API virtual FVector2d GetPositionForNode(const FSchematicGraphNode* InNode) const;
+	ANIMATIONEDITORWIDGETS_API FVector2d GetPositionOffsetForNode(const FGuid& InNodeGuid) const;
+	ANIMATIONEDITORWIDGETS_API virtual FVector2d GetPositionOffsetForNode(const FSchematicGraphNode* InNode) const;
+	ANIMATIONEDITORWIDGETS_API bool GetPositionAnimationEnabledForNode(const FGuid& InNodeGuid) const;
+	ANIMATIONEDITORWIDGETS_API virtual bool GetPositionAnimationEnabledForNode(const FSchematicGraphNode* InNode) const;
+	ANIMATIONEDITORWIDGETS_API FVector2d GetSizeForNode(const FGuid& InNodeGuid) const;
+	ANIMATIONEDITORWIDGETS_API virtual FVector2d GetSizeForNode(const FSchematicGraphNode* InNode) const;
+	ANIMATIONEDITORWIDGETS_API float GetScaleForNode(const FGuid& InNodeGuid) const;
+	ANIMATIONEDITORWIDGETS_API virtual float GetScaleForNode(const FSchematicGraphNode* InNode) const;
+	ANIMATIONEDITORWIDGETS_API float GetScaleOffsetForNode(const FGuid& InNodeGuid) const;
+	ANIMATIONEDITORWIDGETS_API virtual float GetScaleOffsetForNode(const FSchematicGraphNode* InNode) const;
+	ANIMATIONEDITORWIDGETS_API float GetMinimumLinkDistanceForNode(const FGuid& InNodeGuid) const;
+	ANIMATIONEDITORWIDGETS_API virtual float GetMinimumLinkDistanceForNode(const FSchematicGraphNode* InNode) const;
+	ANIMATIONEDITORWIDGETS_API bool IsAutoScaleEnabledForNode(const FGuid& InNodeGuid) const;
+	ANIMATIONEDITORWIDGETS_API virtual bool IsAutoScaleEnabledForNode(const FSchematicGraphNode* InNode) const;
+	ANIMATIONEDITORWIDGETS_API int32 GetNumLayersForNode(const FGuid& InNodeGuid) const;
+	ANIMATIONEDITORWIDGETS_API virtual int32 GetNumLayersForNode(const FSchematicGraphNode* InNode) const;
+	ANIMATIONEDITORWIDGETS_API FLinearColor GetColorForNode(const FGuid& InNodeGuid, int32 InLayerIndex) const;
+	ANIMATIONEDITORWIDGETS_API virtual FLinearColor GetColorForNode(const FSchematicGraphNode* InNode, int32 InLayerIndex) const;
+	ANIMATIONEDITORWIDGETS_API const FSlateBrush* GetBrushForNode(const FGuid& InNodeGuid, int32 InLayerIndex) const;
+	ANIMATIONEDITORWIDGETS_API virtual const FSlateBrush* GetBrushForNode(const FSchematicGraphNode* InNode, int32 InLayerIndex) const;
+	ANIMATIONEDITORWIDGETS_API FText GetToolTipForNode(const FGuid& InNodeGuid) const;
+	ANIMATIONEDITORWIDGETS_API virtual FText GetToolTipForNode(const FSchematicGraphNode* InNode) const;
+	ANIMATIONEDITORWIDGETS_API ESchematicGraphVisibility::Type GetVisibilityForNode(const FGuid& InNodeGuid) const;
+	ANIMATIONEDITORWIDGETS_API virtual ESchematicGraphVisibility::Type GetVisibilityForNode(const FSchematicGraphNode* InNode) const;
+	ANIMATIONEDITORWIDGETS_API TOptional<ESchematicGraphVisibility::Type> GetVisibilityForChildNode(const FGuid& InParentNodeGuid, const FGuid& InChildNodeGuid) const;
+	ANIMATIONEDITORWIDGETS_API virtual TOptional<ESchematicGraphVisibility::Type> GetVisibilityForChildNode(const FSchematicGraphNode* InParentNode, const FSchematicGraphNode* InChildNode) const;
+	ANIMATIONEDITORWIDGETS_API TOptional<FVector2d> GetPositionForChildNode(const FGuid& InParentNodeGuid, const FGuid& InChildNodeGuid) const;
+	ANIMATIONEDITORWIDGETS_API virtual TOptional<FVector2d> GetPositionForChildNode(const FSchematicGraphNode* InParentNode, const FSchematicGraphNode* InChildNode) const;
+	ANIMATIONEDITORWIDGETS_API TOptional<float> GetScaleForChildNode(const FGuid& InParentNodeGuid, const FGuid& InChildNodeGuid) const;
+	ANIMATIONEDITORWIDGETS_API virtual TOptional<float> GetScaleForChildNode(const FSchematicGraphNode* InParentNode, const FSchematicGraphNode* InChildNode) const;
+	ANIMATIONEDITORWIDGETS_API TOptional<bool> GetInteractivityForChildNode(const FGuid& InParentNodeGuid, const FGuid& InChildNodeGuid) const;
+	ANIMATIONEDITORWIDGETS_API virtual TOptional<bool> GetInteractivityForChildNode(const FSchematicGraphNode* InParentNode, const FSchematicGraphNode* InChildNode) const;
+	ANIMATIONEDITORWIDGETS_API bool IsDragSupportedForNode(const FGuid& InNodeGuid) const;
+	ANIMATIONEDITORWIDGETS_API virtual bool IsDragSupportedForNode(const FSchematicGraphNode* InNode) const;
+	ANIMATIONEDITORWIDGETS_API virtual bool GetContextMenuForNode(const FSchematicGraphNode* InNode, FMenuBuilder& OutMenu) const;
+	ANIMATIONEDITORWIDGETS_API const TArray<TSharedPtr<FSchematicGraphNode>> GetSelectedNodes() const;
+	ANIMATIONEDITORWIDGETS_API void ClearSelection();
+
+	ANIMATIONEDITORWIDGETS_API FLinearColor GetBackgroundColorForTag(const FGuid& InNodeGuid, const FGuid& InTagGuid) const;
+	ANIMATIONEDITORWIDGETS_API virtual FLinearColor GetBackgroundColorForTag(const FSchematicGraphTag* InTag) const;
+	ANIMATIONEDITORWIDGETS_API FLinearColor GetForegroundColorForTag(const FGuid& InNodeGuid, const FGuid& InTagGuid) const;
+	ANIMATIONEDITORWIDGETS_API virtual FLinearColor GetForegroundColorForTag(const FSchematicGraphTag* InTag) const;
+	ANIMATIONEDITORWIDGETS_API FLinearColor GetLabelColorForTag(const FGuid& InNodeGuid, const FGuid& InTagGuid) const;
+	ANIMATIONEDITORWIDGETS_API virtual FLinearColor GetLabelColorForTag(const FSchematicGraphTag* InTag) const;
+	ANIMATIONEDITORWIDGETS_API const FSlateBrush* GetBackgroundBrushForTag(const FGuid& InNodeGuid, const FGuid& InTagGuid) const;
+	ANIMATIONEDITORWIDGETS_API virtual const FSlateBrush* GetBackgroundBrushForTag(const FSchematicGraphTag* InTag) const;
+	ANIMATIONEDITORWIDGETS_API const FSlateBrush* GetForegroundBrushForTag(const FGuid& InNodeGuid, const FGuid& InTagGuid) const;
+	ANIMATIONEDITORWIDGETS_API virtual const FSlateBrush* GetForegroundBrushForTag(const FSchematicGraphTag* InTag) const;
+	ANIMATIONEDITORWIDGETS_API const FText GetLabelForTag(const FGuid& InNodeGuid, const FGuid& InTagGuid) const;
+	ANIMATIONEDITORWIDGETS_API virtual const FText GetLabelForTag(const FSchematicGraphTag* InTag) const;
+	ANIMATIONEDITORWIDGETS_API const FText GetToolTipForTag(const FGuid& InNodeGuid, const FGuid& InTagGuid) const;
+	ANIMATIONEDITORWIDGETS_API virtual const FText GetToolTipForTag(const FSchematicGraphTag* InTag) const;
+	ANIMATIONEDITORWIDGETS_API ESchematicGraphVisibility::Type GetVisibilityForTag(const FGuid& InNodeGuid, const FGuid& InTagGuid) const;
+	ANIMATIONEDITORWIDGETS_API virtual ESchematicGraphVisibility::Type GetVisibilityForTag(const FSchematicGraphTag* InTag) const;
 
 	template<typename LinkType = FSchematicGraphLink>
 	LinkType* AddLink(const FGuid& InSourceNodeGuid, const FGuid& InTargetNodeGuid, bool bNotify = true)
@@ -227,37 +227,37 @@ public:
 		return const_cast<LinkType*>(ConstThis->FindLinkChecked<LinkType>(InSourceNodeGuid, InTargetNodeGuid));
 	}
 
-	bool IsLinkedTo(const FGuid& InSourceNodeGuid, const FGuid& InTargetNodeGuid) const;
-	TArray<const FSchematicGraphLink*> FindLinksOnNode(const FGuid& InNodeGuid) const;
-	TArray<const FSchematicGraphLink*> FindLinksOnSource(const FGuid& InSourceNodeGuid) const;
-	TArray<const FSchematicGraphLink*> FindLinksOnTarget(const FGuid& InTargetNodeGuid) const;
+	ANIMATIONEDITORWIDGETS_API bool IsLinkedTo(const FGuid& InSourceNodeGuid, const FGuid& InTargetNodeGuid) const;
+	ANIMATIONEDITORWIDGETS_API TArray<const FSchematicGraphLink*> FindLinksOnNode(const FGuid& InNodeGuid) const;
+	ANIMATIONEDITORWIDGETS_API TArray<const FSchematicGraphLink*> FindLinksOnSource(const FGuid& InSourceNodeGuid) const;
+	ANIMATIONEDITORWIDGETS_API TArray<const FSchematicGraphLink*> FindLinksOnTarget(const FGuid& InTargetNodeGuid) const;
 
-	virtual bool RemoveLink(const FGuid& InLinkGuid);
+	ANIMATIONEDITORWIDGETS_API virtual bool RemoveLink(const FGuid& InLinkGuid);
 
 	const TArray<TSharedPtr<FSchematicGraphLink>>& GetLinks() const { return Links; }
 
-	float GetMinimumForLink(const FGuid& InLinkGuid) const;
-	virtual float GetMinimumForLink(const FSchematicGraphLink* InLink) const;
-	float GetMaximumForLink(const FGuid& InLinkGuid) const;
-	virtual float GetMaximumForLink(const FSchematicGraphLink* InLink) const;
-	FVector2d GetSourceNodeOffsetForLink(const FGuid& InLinkGuid) const;
-	FVector2d GetSourceNodeOffsetForLink(const FSchematicGraphLink* InLink) const;
-	FVector2d GetTargetNodeOffsetForLink(const FGuid& InLinkGuid) const;
-	FVector2d GetTargetNodeOffsetForLink(const FSchematicGraphLink* InLink) const;
-	FLinearColor GetColorForLink(const FGuid& InLinkGuid) const;
-	virtual FLinearColor GetColorForLink(const FSchematicGraphLink* InLink) const;
-	float GetThicknessForLink(const FGuid& InLinkGuid) const;
-	virtual float GetThicknessForLink(const FSchematicGraphLink* InLink) const;
-	const FSlateBrush* GetBrushForLink(const FGuid& InLinkGuid) const;
-	virtual const FSlateBrush* GetBrushForLink(const FSchematicGraphLink* InLink) const;
-	const FText GetToolTipForLink(const FGuid& InLinkGuid) const;
-	virtual const FText GetToolTipForLink(const FSchematicGraphLink* InLink) const;
-	ESchematicGraphVisibility::Type GetVisibilityForLink(const FGuid& InLinkGuid) const;
-	virtual ESchematicGraphVisibility::Type GetVisibilityForLink(const FSchematicGraphLink* InLink) const;
+	ANIMATIONEDITORWIDGETS_API float GetMinimumForLink(const FGuid& InLinkGuid) const;
+	ANIMATIONEDITORWIDGETS_API virtual float GetMinimumForLink(const FSchematicGraphLink* InLink) const;
+	ANIMATIONEDITORWIDGETS_API float GetMaximumForLink(const FGuid& InLinkGuid) const;
+	ANIMATIONEDITORWIDGETS_API virtual float GetMaximumForLink(const FSchematicGraphLink* InLink) const;
+	ANIMATIONEDITORWIDGETS_API FVector2d GetSourceNodeOffsetForLink(const FGuid& InLinkGuid) const;
+	ANIMATIONEDITORWIDGETS_API FVector2d GetSourceNodeOffsetForLink(const FSchematicGraphLink* InLink) const;
+	ANIMATIONEDITORWIDGETS_API FVector2d GetTargetNodeOffsetForLink(const FGuid& InLinkGuid) const;
+	ANIMATIONEDITORWIDGETS_API FVector2d GetTargetNodeOffsetForLink(const FSchematicGraphLink* InLink) const;
+	ANIMATIONEDITORWIDGETS_API FLinearColor GetColorForLink(const FGuid& InLinkGuid) const;
+	ANIMATIONEDITORWIDGETS_API virtual FLinearColor GetColorForLink(const FSchematicGraphLink* InLink) const;
+	ANIMATIONEDITORWIDGETS_API float GetThicknessForLink(const FGuid& InLinkGuid) const;
+	ANIMATIONEDITORWIDGETS_API virtual float GetThicknessForLink(const FSchematicGraphLink* InLink) const;
+	ANIMATIONEDITORWIDGETS_API const FSlateBrush* GetBrushForLink(const FGuid& InLinkGuid) const;
+	ANIMATIONEDITORWIDGETS_API virtual const FSlateBrush* GetBrushForLink(const FSchematicGraphLink* InLink) const;
+	ANIMATIONEDITORWIDGETS_API const FText GetToolTipForLink(const FGuid& InLinkGuid) const;
+	ANIMATIONEDITORWIDGETS_API virtual const FText GetToolTipForLink(const FSchematicGraphLink* InLink) const;
+	ANIMATIONEDITORWIDGETS_API ESchematicGraphVisibility::Type GetVisibilityForLink(const FGuid& InLinkGuid) const;
+	ANIMATIONEDITORWIDGETS_API virtual ESchematicGraphVisibility::Type GetVisibilityForLink(const FSchematicGraphLink* InLink) const;
 
 	virtual bool IsAutoGroupingEnabled() const { return true; }
 	virtual float GetAutoGroupingDistance() const { return 4.f; }
-	virtual FSchematicGraphGroupNode* AddAutoGroupNode(); 
+	ANIMATIONEDITORWIDGETS_API virtual FSchematicGraphGroupNode* AddAutoGroupNode(); 
 
 	FOnSchematicGraphNodeAdded& OnNodeAdded() { return OnNodeAddedDelegate; }
 	FOnNodeRemoved& OnNodeRemoved() { return OnNodeRemovedDelegate; }
@@ -269,14 +269,14 @@ public:
 
 	virtual bool GetForwardedNodeForDrag(FGuid& InOutGuid) const { return false; }
 
-	const FSchematicGraphGroupNode* GetLastExpandedNode() const;
-	virtual void SetLastExpandedNode(const FSchematicGraphGroupNode* InGroupNode);
+	ANIMATIONEDITORWIDGETS_API const FSchematicGraphGroupNode* GetLastExpandedNode() const;
+	ANIMATIONEDITORWIDGETS_API virtual void SetLastExpandedNode(const FSchematicGraphGroupNode* InGroupNode);
 	void ClearLastExpandedNode()
 	{
 		SetLastExpandedNode(nullptr);
 	}
 
-	virtual void Tick(float InDeltaTime);
+	ANIMATIONEDITORWIDGETS_API virtual void Tick(float InDeltaTime);
 
 protected:
 
@@ -299,4 +299,4 @@ protected:
 	friend class FSchematicGraphNode;
 };
 
-#endif
\ No newline at end of file
+#endif
diff --git a/Engine/Source/Editor/AnimationEditorWidgets/Public/SchematicGraphPanel/SchematicGraphNode.h b/Engine/Source/Editor/AnimationEditorWidgets/Public/SchematicGraphPanel/SchematicGraphNode.h
--- a/Engine/Source/Editor/AnimationEditorWidgets/Public/SchematicGraphPanel/SchematicGraphNode.h
+++ b/Engine/Source/Editor/AnimationEditorWidgets/Public/SchematicGraphPanel/SchematicGraphNode.h
@@ -13,43 +13,43 @@ SCHEMATICGRAPHELEMENT_BODY(ClassName, SuperClass, FSchematicGraphNode)
 
 class FSchematicGraphGroupNode;
 
-class ANIMATIONEDITORWIDGETS_API FSchematicGraphNode : public TSharedFromThis<FSchematicGraphNode>
+class FSchematicGraphNode : public TSharedFromThis<FSchematicGraphNode>
 {
 public:
 
 	SCHEMATICGRAPHELEMENT_BODY_BASE(FSchematicGraphNode)
 
-	FSchematicGraphNode();
+	ANIMATIONEDITORWIDGETS_API FSchematicGraphNode();
 	
 	const FGuid& GetGuid() const { return Guid; }
 	const FSchematicGraphModel* GetGraph() const { return Model; }
 	bool HasParentNode() const { return ParentNodeGuid.IsValid(); }
 	const FGuid& GetParentNodeGuid() const { return ParentNodeGuid; }
-	const FSchematicGraphNode* GetParentNode() const;
-	FSchematicGraphNode* GetParentNode();
-	const FGuid& GetRootNodeGuid() const;
-	const FSchematicGraphNode* GetRootNode() const;
-	FSchematicGraphNode* GetRootNode();
+	ANIMATIONEDITORWIDGETS_API const FSchematicGraphNode* GetParentNode() const;
+	ANIMATIONEDITORWIDGETS_API FSchematicGraphNode* GetParentNode();
+	ANIMATIONEDITORWIDGETS_API const FGuid& GetRootNodeGuid() const;
+	ANIMATIONEDITORWIDGETS_API const FSchematicGraphNode* GetRootNode() const;
+	ANIMATIONEDITORWIDGETS_API FSchematicGraphNode* GetRootNode();
 	bool IsRootNode() const { return !HasParentNode(); }
-	FSchematicGraphGroupNode* GetGroupNode();
-	const FSchematicGraphGroupNode* GetGroupNode() const;
+	ANIMATIONEDITORWIDGETS_API FSchematicGraphGroupNode* GetGroupNode();
+	ANIMATIONEDITORWIDGETS_API const FSchematicGraphGroupNode* GetGroupNode() const;
 	int32 GetNumChildNodes() const { return ChildNodeGuids.Num(); }
 	const TArray<FGuid>& GetChildNodeGuids() const { return ChildNodeGuids; }
-	const FSchematicGraphNode* GetChildNode(int32 InChildNodeIndex) const;
-	FSchematicGraphNode* GetChildNode(int32 InChildNodeIndex);
-	TOptional<ESchematicGraphVisibility::Type> GetVisibilityForChildNode(const FGuid& InChildGuid) const;
+	ANIMATIONEDITORWIDGETS_API const FSchematicGraphNode* GetChildNode(int32 InChildNodeIndex) const;
+	ANIMATIONEDITORWIDGETS_API FSchematicGraphNode* GetChildNode(int32 InChildNodeIndex);
+	ANIMATIONEDITORWIDGETS_API TOptional<ESchematicGraphVisibility::Type> GetVisibilityForChildNode(const FGuid& InChildGuid) const;
 	virtual TOptional<ESchematicGraphVisibility::Type> GetVisibilityForChildNode(const FSchematicGraphNode* InChildNode) const { return  TOptional<ESchematicGraphVisibility::Type>(); }
-	TOptional<FVector2d> GetPositionForChildNode(const FGuid& InChildGuid) const;
+	ANIMATIONEDITORWIDGETS_API TOptional<FVector2d> GetPositionForChildNode(const FGuid& InChildGuid) const;
 	virtual TOptional<FVector2d> GetPositionForChildNode(const FSchematicGraphNode* InChildNode) const { return TOptional<FVector2d>(); }
-	TOptional<float> GetScaleForChildNode(const FGuid& InChildGuid) const;
+	ANIMATIONEDITORWIDGETS_API TOptional<float> GetScaleForChildNode(const FGuid& InChildGuid) const;
 	virtual TOptional<float> GetScaleForChildNode(const FSchematicGraphNode* InChildNode) const { return TOptional<float>(); }
-	TOptional<bool> GetInteractivityForChildNode(const FGuid& InChildGuid) const;
+	ANIMATIONEDITORWIDGETS_API TOptional<bool> GetInteractivityForChildNode(const FGuid& InChildGuid) const;
 	virtual TOptional<bool> GetInteractivityForChildNode(const FSchematicGraphNode* InChildNode) const { return TOptional<bool>(); }
 	bool IsSelected() const { return bIsSelected; }
 	void SetSelected(bool bSelected = true) { bIsSelected = bSelected; }
 	virtual const FText& GetLabel() const { return Label; }
 	void SetLabel(const FText& InLabel) { Label = InLabel; }
-	virtual FVector2d GetPosition() const;
+	ANIMATIONEDITORWIDGETS_API virtual FVector2d GetPosition() const;
 	virtual void SetPosition(const FVector2d& InPosition) { Position = InPosition; }
 	virtual FVector2d GetPositionOffset() const { return PositionOffset; }
 	virtual void SetPositionOffset(const FVector2d& InPositionOffset) { PositionOffset = InPositionOffset; }
@@ -59,20 +59,20 @@ public:
 	virtual FLinearColor GetColor(int32 InLayerIndex) const { return Colors[InLayerIndex]; }
 	virtual const FSlateBrush* GetBrush(int32 InLayerIndex) const { return Brushes[InLayerIndex]; }
 	virtual bool IsAutoScaleEnabled() const { return false; }
-	virtual const FText& GetToolTip() const;
-	virtual ESchematicGraphVisibility::Type GetVisibility() const;
+	ANIMATIONEDITORWIDGETS_API virtual const FText& GetToolTip() const;
+	ANIMATIONEDITORWIDGETS_API virtual ESchematicGraphVisibility::Type GetVisibility() const;
 	virtual void SetVisibility(ESchematicGraphVisibility::Type InVisibility) { Visibility = InVisibility; }
-	virtual bool IsInteractive() const;
+	ANIMATIONEDITORWIDGETS_API virtual bool IsInteractive() const;
 	virtual bool IsDragSupported() const { return bDragSupported; }
 	virtual void SetDragSupported(bool InDragSupported) { bDragSupported = InDragSupported;}
 
-	virtual FReply OnClicked(const FPointerEvent& InMouseEvent);
-	virtual void OnMouseEnter();
-	virtual void OnMouseLeave();
-	virtual void OnDragOver();
-	virtual void OnDragLeave();;
+	ANIMATIONEDITORWIDGETS_API virtual FReply OnClicked(const FPointerEvent& InMouseEvent);
+	ANIMATIONEDITORWIDGETS_API virtual void OnMouseEnter();
+	ANIMATIONEDITORWIDGETS_API virtual void OnMouseLeave();
+	ANIMATIONEDITORWIDGETS_API virtual void OnDragOver();
+	ANIMATIONEDITORWIDGETS_API virtual void OnDragLeave();;
 
-	virtual FString GetDragDropDecoratorLabel() const;
+	ANIMATIONEDITORWIDGETS_API virtual FString GetDragDropDecoratorLabel() const;
 
 	template<typename TagType = FSchematicGraphTag>
 	TagType* AddTag(bool bNotify = true)
@@ -134,13 +134,13 @@ public:
 		return const_cast<TagType*>(ConstThis->FindTagChecked<TagType>(InTagGuid));
 	}
 
-	virtual bool RemoveTag(const FGuid& InTagGuid);
+	ANIMATIONEDITORWIDGETS_API virtual bool RemoveTag(const FGuid& InTagGuid);
 
 	const TArray<TSharedPtr<FSchematicGraphTag>>& GetTags() const { return Tags; }
 
 protected:
 
-	void NotifyTagAdded(const TSharedPtr<FSchematicGraphTag>& Tag);
+	ANIMATIONEDITORWIDGETS_API void NotifyTagAdded(const TSharedPtr<FSchematicGraphTag>& Tag);
 
 	FSchematicGraphModel* Model = nullptr;
 	FGuid Guid = FGuid::NewGuid();
@@ -166,40 +166,40 @@ protected:
 	friend class FSchematicGraphModel;
 };
 
-class ANIMATIONEDITORWIDGETS_API FSchematicGraphGroupNode : public FSchematicGraphNode
+class FSchematicGraphGroupNode : public FSchematicGraphNode
 {
 public:
 
 	SCHEMATICGRAPHNODE_BODY(FSchematicGraphGroupNode, FSchematicGraphNode)
 
-	FSchematicGraphGroupNode();
+	ANIMATIONEDITORWIDGETS_API FSchematicGraphGroupNode();
 	virtual ~FSchematicGraphGroupNode() override {}
 
 	// FSchematicGraphNode interface
-	virtual TOptional<ESchematicGraphVisibility::Type> GetVisibilityForChildNode(const FSchematicGraphNode* InChildNode) const override;
-	virtual TOptional<FVector2d> GetPositionForChildNode(const FSchematicGraphNode* InChildNode) const override;
-	virtual TOptional<float> GetScaleForChildNode(const FSchematicGraphNode* InChildNode) const override;
-	virtual TOptional<bool> GetInteractivityForChildNode(const FSchematicGraphNode* InChildNode) const override;
-	virtual FReply OnClicked(const FPointerEvent& InMouseEvent) override;
-	virtual void OnDragOver() override;
-	virtual void OnDragLeave() override;
-
-	bool IsExpanded() const;
-	bool IsExpanding() const;
-	bool IsCollapsing() const;
-	float GetExpansionState() const;
-	void SetExpanded(bool InExpanded, bool bAutoCloseParentGroups = false);
+	ANIMATIONEDITORWIDGETS_API virtual TOptional<ESchematicGraphVisibility::Type> GetVisibilityForChildNode(const FSchematicGraphNode* InChildNode) const override;
+	ANIMATIONEDITORWIDGETS_API virtual TOptional<FVector2d> GetPositionForChildNode(const FSchematicGraphNode* InChildNode) const override;
+	ANIMATIONEDITORWIDGETS_API virtual TOptional<float> GetScaleForChildNode(const FSchematicGraphNode* InChildNode) const override;
+	ANIMATIONEDITORWIDGETS_API virtual TOptional<bool> GetInteractivityForChildNode(const FSchematicGraphNode* InChildNode) const override;
+	ANIMATIONEDITORWIDGETS_API virtual FReply OnClicked(const FPointerEvent& InMouseEvent) override;
+	ANIMATIONEDITORWIDGETS_API virtual void OnDragOver() override;
+	ANIMATIONEDITORWIDGETS_API virtual void OnDragLeave() override;
+
+	ANIMATIONEDITORWIDGETS_API bool IsExpanded() const;
+	ANIMATIONEDITORWIDGETS_API bool IsExpanding() const;
+	ANIMATIONEDITORWIDGETS_API bool IsCollapsing() const;
+	ANIMATIONEDITORWIDGETS_API float GetExpansionState() const;
+	ANIMATIONEDITORWIDGETS_API void SetExpanded(bool InExpanded, bool bAutoCloseParentGroups = false);
 	FLinearColor GetExpansionColor() const { return ExpansionColor; }
 	void SetExpansionColor(FLinearColor InExpansionColor) { ExpansionColor = InExpansionColor; }
 	float GetExpansionRadius() const { return ExpansionRadius; }
 	void SetExpansionRadius(float InExpansionRadius) { ExpansionRadius = InExpansionRadius; }
-	float GetDelayDuration(bool bEnter) const;
+	ANIMATIONEDITORWIDGETS_API float GetDelayDuration(bool bEnter) const;
 	void SetDelayDuration(bool bEnter, float InDelayDuration) { (bEnter ? EnterDelayDuration : LeaveDelayDuration) = InDelayDuration; }
 
 protected:
 
 	const TEasingAttributeInterpolator<float>::FSettings& GetAnimationSettings() const { return AnimationSettings; }
-	void SetAnimationSettings(const TEasingAttributeInterpolator<float>::FSettings& InSettings);
+	ANIMATIONEDITORWIDGETS_API void SetAnimationSettings(const TEasingAttributeInterpolator<float>::FSettings& InSettings);
 
 	float ExpansionRadius = 60.f;
 	FLinearColor ExpansionColor = FLinearColor::White * FLinearColor(1.f, 1.f, 1.f, 0.5f);
@@ -212,13 +212,13 @@ protected:
 	friend class SSchematicGraphNode;
 };
 
-class ANIMATIONEDITORWIDGETS_API FSchematicGraphAutoGroupNode : public FSchematicGraphGroupNode
+class FSchematicGraphAutoGroupNode : public FSchematicGraphGroupNode
 {
 public:
 
 	SCHEMATICGRAPHNODE_BODY(FSchematicGraphAutoGroupNode, FSchematicGraphGroupNode)
 
-	virtual const FText& GetLabel() const override;
+	ANIMATIONEDITORWIDGETS_API virtual const FText& GetLabel() const override;
 
 protected:
 
@@ -226,4 +226,4 @@ protected:
 };
 
 
-#endif
\ No newline at end of file
+#endif
diff --git a/Engine/Source/Editor/AnimationEditorWidgets/Public/SchematicGraphPanel/SchematicGraphStyle.h b/Engine/Source/Editor/AnimationEditorWidgets/Public/SchematicGraphPanel/SchematicGraphStyle.h
--- a/Engine/Source/Editor/AnimationEditorWidgets/Public/SchematicGraphPanel/SchematicGraphStyle.h
+++ b/Engine/Source/Editor/AnimationEditorWidgets/Public/SchematicGraphPanel/SchematicGraphStyle.h
@@ -8,7 +8,9 @@
 #include "Styling/StyleColors.h"
 #include "Brushes/SlateRoundedBoxBrush.h"
 
-class ANIMATIONEDITORWIDGETS_API FSchematicGraphStyle final
+#define UE_API ANIMATIONEDITORWIDGETS_API
+
+class FSchematicGraphStyle final
 	: public FSlateStyleSet
 {
 
@@ -43,10 +45,12 @@ public:
 		FSlateStyleRegistry::RegisterSlateStyle(*this);
 	}
 
-	static FSchematicGraphStyle& Get();
+	static UE_API FSchematicGraphStyle& Get();
 	
 	~FSchematicGraphStyle()
 	{
 		FSlateStyleRegistry::UnRegisterSlateStyle(*this);
 	}
 };
+
+#undef UE_API
diff --git a/Engine/Source/Editor/AnimationEditorWidgets/Public/SchematicGraphPanel/SchematicGraphTag.h b/Engine/Source/Editor/AnimationEditorWidgets/Public/SchematicGraphPanel/SchematicGraphTag.h
--- a/Engine/Source/Editor/AnimationEditorWidgets/Public/SchematicGraphPanel/SchematicGraphTag.h
+++ b/Engine/Source/Editor/AnimationEditorWidgets/Public/SchematicGraphPanel/SchematicGraphTag.h
@@ -8,15 +8,15 @@
 #define SCHEMATICGRAPHTAG_BODY(ClassName, SuperClass) \
 SCHEMATICGRAPHELEMENT_BODY(ClassName, SuperClass, FSchematicGraphTag)
 
-class ANIMATIONEDITORWIDGETS_API FSchematicGraphTag : public TSharedFromThis<FSchematicGraphTag>
+class FSchematicGraphTag : public TSharedFromThis<FSchematicGraphTag>
 {
 public:
 	
 	SCHEMATICGRAPHELEMENT_BODY_BASE(FSchematicGraphTag)
 
-	FSchematicGraphTag();
+	ANIMATIONEDITORWIDGETS_API FSchematicGraphTag();
 
-	const FSchematicGraphNode* GetNode() const;
+	ANIMATIONEDITORWIDGETS_API const FSchematicGraphNode* GetNode() const;
 	const FGuid& GetGuid() const { return Guid; }
 	virtual FLinearColor GetBackgroundColor() const { return BackgroundColor; }
 	virtual void SetBackgroundColor(const FLinearColor& InBackgroundColor) { BackgroundColor = InBackgroundColor; }
@@ -55,7 +55,7 @@ protected:
 	friend class FSchematicGraphNode;
 };
 
-class ANIMATIONEDITORWIDGETS_API FSchematicGraphGroupTag : public FSchematicGraphTag
+class FSchematicGraphGroupTag : public FSchematicGraphTag
 {
 public:
 
@@ -63,8 +63,8 @@ public:
 
 	virtual ~FSchematicGraphGroupTag() override {}
 
-	virtual ESchematicGraphVisibility::Type GetVisibility() const override;
-	virtual FText GetLabel() const override;
+	ANIMATIONEDITORWIDGETS_API virtual ESchematicGraphVisibility::Type GetVisibility() const override;
+	ANIMATIONEDITORWIDGETS_API virtual FText GetLabel() const override;
 };
 
-#endif
\ No newline at end of file
+#endif
diff --git a/Engine/Source/Editor/AnimationEditorWidgets/Public/SchematicGraphPanel/SchematicViewportTabSummoner.h b/Engine/Source/Editor/AnimationEditorWidgets/Public/SchematicGraphPanel/SchematicViewportTabSummoner.h
--- a/Engine/Source/Editor/AnimationEditorWidgets/Public/SchematicGraphPanel/SchematicViewportTabSummoner.h
+++ b/Engine/Source/Editor/AnimationEditorWidgets/Public/SchematicGraphPanel/SchematicViewportTabSummoner.h
@@ -11,7 +11,7 @@
 DECLARE_DELEGATE_OneParam(FOnSchematicViewportCreated, const TSharedRef<class SSchematicGraphPanel>&);
 
 /** Arguments used to create a persona viewport tab */
-struct ANIMATIONEDITORWIDGETS_API FSchematicViewportArgs
+struct FSchematicViewportArgs
 {
 	FSchematicViewportArgs()
 	{}
@@ -24,18 +24,18 @@ struct ANIMATIONEDITORWIDGETS_API FSchematicViewportArgs
 };
 
 
-struct ANIMATIONEDITORWIDGETS_API FSchematicViewportTabSummoner : public FWorkflowTabFactory
+struct FSchematicViewportTabSummoner : public FWorkflowTabFactory
 {
 public:
-	static const FName TabID;
+	static ANIMATIONEDITORWIDGETS_API const FName TabID;
 	
 public:
-	FSchematicViewportTabSummoner(TSharedPtr<class FAssetEditorToolkit> InHostingApp, const FSchematicViewportArgs& InArgs);
+	ANIMATIONEDITORWIDGETS_API FSchematicViewportTabSummoner(TSharedPtr<class FAssetEditorToolkit> InHostingApp, const FSchematicViewportArgs& InArgs);
 	
-	virtual FTabSpawnerEntry& RegisterTabSpawner(TSharedRef<FTabManager> TabManager, const FApplicationMode* CurrentApplicationMode) const;
-	virtual TSharedRef<SWidget> CreateTabBody(const FWorkflowTabSpawnInfo& Info) const override;
+	ANIMATIONEDITORWIDGETS_API virtual FTabSpawnerEntry& RegisterTabSpawner(TSharedRef<FTabManager> TabManager, const FApplicationMode* CurrentApplicationMode) const;
+	ANIMATIONEDITORWIDGETS_API virtual TSharedRef<SWidget> CreateTabBody(const FWorkflowTabSpawnInfo& Info) const override;
 
 	TWeakPtr<FSchematicGraphModel> SchematicGraphWeak;
 	FOnSchematicViewportCreated OnViewportCreated;
 };
-#endif
\ No newline at end of file
+#endif
diff --git a/Engine/Source/Editor/AnimationModifiers/Private/SAnimationModifiersTab.h b/Engine/Source/Editor/AnimationModifiers/Private/SAnimationModifiersTab.h
--- a/Engine/Source/Editor/AnimationModifiers/Private/SAnimationModifiersTab.h
+++ b/Engine/Source/Editor/AnimationModifiers/Private/SAnimationModifiersTab.h
@@ -11,6 +11,8 @@
 #include "Widgets/DeclarativeSyntaxSupport.h"
 #include "Widgets/SCompoundWidget.h"
 
+#define UE_API ANIMATIONMODIFIERS_API
+
 class IAssetEditorInstance;
 class IDetailsView;
 class SMenuAnchor;
@@ -23,7 +25,7 @@ class UObject;
 class USkeleton;
 struct FGeometry;
 
-class ANIMATIONMODIFIERS_API SAnimationModifiersTab : public SCompoundWidget, public FEditorUndoClient
+class SAnimationModifiersTab : public SCompoundWidget, public FEditorUndoClient
 {
 public:
 	SLATE_BEGIN_ARGS(SAnimationModifiersTab)
@@ -31,56 +33,56 @@ public:
 	SLATE_ARGUMENT(TWeakPtr<class FAssetEditorToolkit>, InHostingApp)
 	SLATE_END_ARGS()
 
-	SAnimationModifiersTab();
-	~SAnimationModifiersTab();
+	UE_API SAnimationModifiersTab();
+	UE_API ~SAnimationModifiersTab();
 
 	/** SWidget functions */
-	void Construct(const FArguments& InArgs);
+	UE_API void Construct(const FArguments& InArgs);
 
 	/** Begin SCompoundWidget */
-	virtual void Tick(const FGeometry& AllottedGeometry, const double InCurrentTime, const float InDeltaTime) override;
+	UE_API virtual void Tick(const FGeometry& AllottedGeometry, const double InCurrentTime, const float InDeltaTime) override;
 	/** End SCompoundWidget */
 
 	/** Begin FEditorUndoClient */
-	virtual void PostUndo(bool bSuccess) override;
-	virtual void PostRedo(bool bSuccess) override;
+	UE_API virtual void PostUndo(bool bSuccess) override;
+	UE_API virtual void PostRedo(bool bSuccess) override;
 	/** End FEditorUndoClient */
 protected:
 	/** Callback for when user has picked a modifier to add */
-	void OnModifierPicked(UClass* PickedClass);
+	UE_API void OnModifierPicked(UClass* PickedClass);
 
-	void CreateInstanceDetailsView();	
+	UE_API void CreateInstanceDetailsView();	
 
 	/** UI apply all modifiers button callback */
-	FReply OnApplyAllModifiersClicked();
+	UE_API FReply OnApplyAllModifiersClicked();
 
 	/** Callbacks for available modifier actions */
-	void OnApplyModifier(const TArray<TWeakObjectPtr<UAnimationModifier>>& Instances);	
-	void OnRevertModifier(const TArray<TWeakObjectPtr<UAnimationModifier>>& Instances);
-	bool OnCanRevertModifier(const TArray<TWeakObjectPtr<UAnimationModifier>>& Instances);	
-	void OnRemoveModifier(const TArray<TWeakObjectPtr<UAnimationModifier>>& Instances);
-	void OnOpenModifier(const TWeakObjectPtr<UAnimationModifier>& Instance);
+	UE_API void OnApplyModifier(const TArray<TWeakObjectPtr<UAnimationModifier>>& Instances);	
+	UE_API void OnRevertModifier(const TArray<TWeakObjectPtr<UAnimationModifier>>& Instances);
+	UE_API bool OnCanRevertModifier(const TArray<TWeakObjectPtr<UAnimationModifier>>& Instances);	
+	UE_API void OnRemoveModifier(const TArray<TWeakObjectPtr<UAnimationModifier>>& Instances);
+	UE_API void OnOpenModifier(const TWeakObjectPtr<UAnimationModifier>& Instance);
 
-	void OnMoveModifierUp(const TWeakObjectPtr<UAnimationModifier>& Instance);
-	void OnMoveModifierDown(const TWeakObjectPtr<UAnimationModifier>& Instance);
+	UE_API void OnMoveModifierUp(const TWeakObjectPtr<UAnimationModifier>& Instance);
+	UE_API void OnMoveModifierDown(const TWeakObjectPtr<UAnimationModifier>& Instance);
 
 	/** Flags UI dirty and will refresh during the next Tick*/
-	void Refresh();
+	UE_API void Refresh();
 
 	/** Callback for compiled blueprints, this ensures to refresh the UI */
-	void OnBlueprintCompiled(UBlueprint* Blueprint);
+	UE_API void OnBlueprintCompiled(UBlueprint* Blueprint);
 	/** Callback to keep track of when an asset is opened, this is necessary for when an editor document tab is reused and this tab isn't recreated */
-	void OnAssetOpened(UObject* Object, IAssetEditorInstance* Instance);
+	UE_API void OnAssetOpened(UObject* Object, IAssetEditorInstance* Instance);
 
 	/** Applying and reverting of modifiers */
-	void ApplyModifiers(const TArray<UAnimationModifier*>& Modifiers);
-	void RevertModifiers(const TArray<UAnimationModifier*>& Modifiers);
+	UE_API void ApplyModifiers(const TArray<UAnimationModifier*>& Modifiers);
+	UE_API void RevertModifiers(const TArray<UAnimationModifier*>& Modifiers);
 
 	/** Retrieves the currently opened animation asset type and modifier user data */
-	void RetrieveAnimationAsset();
+	UE_API void RetrieveAnimationAsset();
 
 	/** Retrieves all animation sequences which are dependent on the current opened skeleton */
-	void FindAnimSequencesForSkeleton(TArray<UAnimSequence *> &ReferencedAnimSequences);
+	UE_API void FindAnimSequencesForSkeleton(TArray<UAnimSequence *> &ReferencedAnimSequences);
 protected:
 	TWeakPtr<class FAssetEditorToolkit> HostingApp;
 
@@ -100,6 +102,8 @@ protected:
 	TSharedPtr<SModifierListView> ModifierListView;
 	TSharedPtr<SMenuAnchor> AddModifierCombobox;	
 private:
-	void RetrieveModifierData();
-	void ResetModifierData();
+	UE_API void RetrieveModifierData();
+	UE_API void ResetModifierData();
 };
+
+#undef UE_API
diff --git a/Engine/Source/Editor/AnimationModifiers/Public/AnimationModifier.h b/Engine/Source/Editor/AnimationModifiers/Public/AnimationModifier.h
--- a/Engine/Source/Editor/AnimationModifiers/Public/AnimationModifier.h
+++ b/Engine/Source/Editor/AnimationModifiers/Public/AnimationModifier.h
@@ -20,6 +20,8 @@
 
 #include "AnimationModifier.generated.h"
 
+#define UE_API ANIMATIONMODIFIERS_API
+
 class FArchive;
 class FName;
 class UAnimSequence;
@@ -28,54 +30,54 @@ class USkeleton;
 struct FFrame;
 struct FObjectKey;
 
-UCLASS(Blueprintable, config = Editor, defaultconfig)
-class ANIMATIONMODIFIERS_API UAnimationModifier : public UObject
+UCLASS(MinimalAPI, Blueprintable, config = Editor, defaultconfig)
+class UAnimationModifier : public UObject
 {
 	GENERATED_BODY()
 
 	friend class FAnimationModifierDetailCustomization;
 public:
-	UAnimationModifier();
+	UE_API UAnimationModifier();
 
 	/** Applying and reverting the modifier for the given Animation Sequence */
-	void ApplyToAnimationSequence(UAnimSequence* AnimSequence) const;
-	void RevertFromAnimationSequence(UAnimSequence* AnimSequence) const;
+	UE_API void ApplyToAnimationSequence(UAnimSequence* AnimSequence) const;
+	UE_API void RevertFromAnimationSequence(UAnimSequence* AnimSequence) const;
 
 	/** Executed when the Animation is initialized (native event for debugging / testing purposes) */
 	UFUNCTION(BlueprintNativeEvent)
-	void OnApply(UAnimSequence* AnimationSequence);
+	UE_API void OnApply(UAnimSequence* AnimationSequence);
 	virtual void OnApply_Implementation(UAnimSequence* AnimationSequence) {}
 	
 	UFUNCTION(BlueprintNativeEvent)
-	void OnRevert(UAnimSequence* AnimationSequence);
+	UE_API void OnRevert(UAnimSequence* AnimationSequence);
 	virtual void OnRevert_Implementation(UAnimSequence* AnimationSequence) {}
 
 	/** Returns whether or not this modifier can be reverted, which means it will have to been applied to the given AnimationSequence */
-	bool CanRevert(IInterface_AssetUserData* AssetUserDataInterface) const;
+	UE_API bool CanRevert(IInterface_AssetUserData* AssetUserDataInterface) const;
 
 	/** Whether or not the latest compiled version of the blueprint is applied for this instance */
-	bool IsLatestRevisionApplied(IInterface_AssetUserData* AssetUserDataInterface) const;
+	UE_API bool IsLatestRevisionApplied(IInterface_AssetUserData* AssetUserDataInterface) const;
 
 	//! @brief Get the latest revision GUID of this modifier class
-	FGuid GetLatestRevisionGuid() const;
+	UE_API FGuid GetLatestRevisionGuid() const;
 
 	//! @brief Get Asset Registry tags for applied modifiers from _skeleton_
 	//! @param[out] OutTags AnimationModifiers = %PATH%=%REVISION%;...
-	static void GetAssetRegistryTagsForAppliedModifiersFromSkeleton(FAssetRegistryTagsContext Context);
+	static UE_API void GetAssetRegistryTagsForAppliedModifiersFromSkeleton(FAssetRegistryTagsContext Context);
 
-	static const FName AnimationModifiersTag;
+	static UE_API const FName AnimationModifiersTag;
 	static constexpr TCHAR AnimationModifiersDelimiter = ';';
 	static constexpr TCHAR AnimationModifiersAssignment = '=';
 
 	//! @brief Check if this Modifier On Skeleton has PreviouslyAppliedModifier_DEPRECATED from previous version
 	bool HasLegacyPreviousAppliedModifierOnSkeleton() const { return GetLegacyPreviouslyAppliedModifierForModifierOnSkeleton() != nullptr; }
 	//! @brief Mark this modifier on skeleton as reverted (affect CanRevert, IsLatestRevisionApplied)
-	void RemoveLegacyPreviousAppliedModifierOnSkeleton(USkeleton* Skeleton);
+	UE_API void RemoveLegacyPreviousAppliedModifierOnSkeleton(USkeleton* Skeleton);
 
 	// Begin UObject Overrides
-	virtual void PostInitProperties() override;
-	virtual void Serialize(FArchive& Ar) override;
-	virtual void PostLoad() override;
+	UE_API virtual void PostInitProperties() override;
+	UE_API virtual void Serialize(FArchive& Ar) override;
+	UE_API virtual void PostLoad() override;
 	// End UObject Overrides
 
 	/** If this is set to true then the animation modifier will call it's reapply function after any change made to the owning asset. */
@@ -88,42 +90,42 @@ public:
 
 protected:
 	// Derived class accessors to skeleton and anim sequence 
-	const UAnimSequence* GetAnimationSequence();
-	const USkeleton* GetSkeleton();
+	UE_API const UAnimSequence* GetAnimationSequence();
+	UE_API const USkeleton* GetSkeleton();
 
 	/** Used for natively updating the revision GUID, fairly basic and relies on config files currently */
-	virtual int32 GetNativeClassRevision() const;
+	UE_API virtual int32 GetNativeClassRevision() const;
 
 	/** Checks if the animation data has to be re-baked / compressed and does so */
-	void UpdateCompressedAnimationData();
+	UE_API void UpdateCompressedAnimationData();
 
 	/** Updating of blueprint and native GUIDs*/
-	void UpdateRevisionGuid(UClass* ModifierClass);
-	void UpdateNativeRevisionGuid();
+	UE_API void UpdateRevisionGuid(UClass* ModifierClass);
+	UE_API void UpdateNativeRevisionGuid();
 
-	void ExecuteOnRevert(UAnimSequence* InAnimSequence);
-	void ExecuteOnApply(UAnimSequence* InAnimSequence);
+	UE_API void ExecuteOnRevert(UAnimSequence* InAnimSequence);
+	UE_API void ExecuteOnApply(UAnimSequence* InAnimSequence);
 	
 	/** Applies all instances of the provided Modifier class to its outer Animation Sequence*/
-	static void ApplyToAll(TSubclassOf<UAnimationModifier> ModifierSubClass, bool bForceApply = true);
-	static void LoadModifierReferencers(TSubclassOf<UAnimationModifier> ModifierSubClass);
+	static UE_API void ApplyToAll(TSubclassOf<UAnimationModifier> ModifierSubClass, bool bForceApply = true);
+	static UE_API void LoadModifierReferencers(TSubclassOf<UAnimationModifier> ModifierSubClass);
 private:
 	// These value were not set during OnRevert, prefer to use the input from OnApply() or OnRevert()
 	UAnimSequence* CurrentAnimSequence = nullptr;
 	USkeleton* CurrentSkeleton = nullptr;
 
 	//! @brief Try get the applied modifier instance (associated with this modifier) on the animation sequence, null if not applied on the sequence
-	UAnimationModifier* GetAppliedModifier(IInterface_AssetUserData* AssetUserDataInterface) const;
+	UE_API UAnimationModifier* GetAppliedModifier(IInterface_AssetUserData* AssetUserDataInterface) const;
 
 	//! @brief Set the applied modifier instance on the Animation Sequence
-	void SetAppliedModifier(TScriptInterface<IInterface_AssetUserData> AssetUserDataInterface, UAnimationModifier* AppliedModifier) const;
+	UE_API void SetAppliedModifier(TScriptInterface<IInterface_AssetUserData> AssetUserDataInterface, UAnimationModifier* AppliedModifier) const;
 	//! @brief Remove the applied modifier instance (associated with this modifier) from the Animation Sequence
 	//! @return The modifier instance removed
-	UAnimationModifier* FindAndRemoveAppliedModifier(TScriptInterface<IInterface_AssetUserData> AssetUserDataInterface) const;
+	UE_API UAnimationModifier* FindAndRemoveAppliedModifier(TScriptInterface<IInterface_AssetUserData> AssetUserDataInterface) const;
 	//! @brief Get the applied modifier instance for modifier on skeleton read from _legacy_ version
 	//! @note This function is only intended for backward compatibility with existing asset
 	//!	@return PreviouslyAppliedModifier_DEPRECATED on Modifier on Skeleton, _before_ upgrade
-	UAnimationModifier* GetLegacyPreviouslyAppliedModifierForModifierOnSkeleton() const;
+	UE_API UAnimationModifier* GetLegacyPreviouslyAppliedModifierForModifierOnSkeleton() const;
 
 
 	// - On class-default-objects, this is the latest revision GUID
@@ -150,7 +152,7 @@ namespace UE
 	namespace Anim
 	{
 		// RAII object to determine how Animation Modifier warning/errors are handled (ignore, user dialog, etc.)
-		struct ANIMATIONMODIFIERS_API FApplyModifiersScope
+		struct FApplyModifiersScope
 		{
 			friend UAnimationModifier;
 
@@ -192,25 +194,27 @@ namespace UE
 			static bool IsScopePending() { return !ScopeModeStack.IsEmpty(); }
 protected:
 			/** Determine how to handle an Animation Modifier error, and execute accordingly */
-			static void HandleError(const UAnimationModifier* Modifier, const FText& Message, const FText& Title);
+			static UE_API void HandleError(const UAnimationModifier* Modifier, const FText& Message, const FText& Title);
 
 			/** Determine how to handle an Animation Modifier warning, and execute accordingly. Returns whether or not the warning was handled (true) or warrants reverting the applied Animation Modifier (false) */
-			static bool HandleWarning(const UAnimationModifier* Modifier, const FText& Message, const FText& Title);
+			static UE_API bool HandleWarning(const UAnimationModifier* Modifier, const FText& Message, const FText& Title);
 
 private:
 			// Open a scope to control error handling when batch applying animation modifiers
-			static ESuppressionMode Open(ESuppressionMode Mode = NoChange);
+			static UE_API ESuppressionMode Open(ESuppressionMode Mode = NoChange);
 			// Close the most recent scope
-			static void Close();
+			static UE_API void Close();
 
-			static ESuppressionMode CurrentMode();
+			static UE_API ESuppressionMode CurrentMode();
 
 			/** Errors already acknowledged */
-			static TSet<FObjectKey> ErrorResponse;
+			static UE_API TSet<FObjectKey> ErrorResponse;
 			/** Warnings already ignored or treated as error */
-			static TMap<FObjectKey, EAppReturnType::Type> WarningResponse;
+			static UE_API TMap<FObjectKey, EAppReturnType::Type> WarningResponse;
 			/** Error handle mode stack for scopes */
-			static TArray<ESuppressionMode, TInlineAllocator<4>> ScopeModeStack;
+			static UE_API TArray<ESuppressionMode, TInlineAllocator<4>> ScopeModeStack;
 		};
 	}
 }
+
+#undef UE_API
diff --git a/Engine/Source/Editor/AnimationModifiers/Public/AnimationModifiersAssetUserData.h b/Engine/Source/Editor/AnimationModifiers/Public/AnimationModifiersAssetUserData.h
--- a/Engine/Source/Editor/AnimationModifiers/Public/AnimationModifiersAssetUserData.h
+++ b/Engine/Source/Editor/AnimationModifiers/Public/AnimationModifiersAssetUserData.h
@@ -12,13 +12,15 @@
 
 #include "AnimationModifiersAssetUserData.generated.h"
 
+#define UE_API ANIMATIONMODIFIERS_API
+
 class FArchive;
 class UObject;
 template <typename T> struct TObjectPtr;
 
 /** Asset user data which can be added to a USkeleton or UAnimSequence to keep track of Animation Modifiers */
-UCLASS()
-class ANIMATIONMODIFIERS_API UAnimationModifiersAssetUserData : public UAssetUserData
+UCLASS(MinimalAPI)
+class UAnimationModifiersAssetUserData : public UAssetUserData
 {
 	GENERATED_BODY()
 
@@ -28,26 +30,26 @@ class ANIMATIONMODIFIERS_API UAnimationModifiersAssetUserData : public UAssetUse
 	friend class FAnimationModifiersModule;
 	friend class UAnimationModifier;
 public:
-	UAnimationModifiersAssetUserData(const FObjectInitializer& ObjectInitializer);
-	const TArray<UAnimationModifier*>& GetAnimationModifierInstances() const;
+	UE_API UAnimationModifiersAssetUserData(const FObjectInitializer& ObjectInitializer);
+	UE_API const TArray<UAnimationModifier*>& GetAnimationModifierInstances() const;
 	
 	UFUNCTION(BlueprintCallable, Category=Animation, meta = (ScriptMethod))
-	static bool AddAnimationModifierOfClass(UAnimSequenceBase* AnimationSequenceBase, TSubclassOf<UAnimationModifier> InModifierClass);
+	static UE_API bool AddAnimationModifierOfClass(UAnimSequenceBase* AnimationSequenceBase, TSubclassOf<UAnimationModifier> InModifierClass);
 protected:	 
 	/** Begin UAssetUserData overrides */
-	virtual void PostLoad() override;
-	virtual void PostEditChangeProperty(struct FPropertyChangedEvent& PropertyChangedEvent) override;	
-	virtual void Serialize(FArchive& Ar) override;
+	UE_API virtual void PostLoad() override;
+	UE_API virtual void PostEditChangeProperty(struct FPropertyChangedEvent& PropertyChangedEvent) override;	
+	UE_API virtual void Serialize(FArchive& Ar) override;
 #if WITH_EDITOR
-	virtual void PostEditChangeOwner(const FPropertyChangedEvent& PropertyChangedEvent) override;
+	UE_API virtual void PostEditChangeOwner(const FPropertyChangedEvent& PropertyChangedEvent) override;
 #endif // WITH_EDITOR
 	/** End UAssetUserData overrides */
 
-	void AddAnimationModifier(UAnimationModifier* Instance);
-	void RemoveAnimationModifierInstance(UAnimationModifier* Instance);
-	void ChangeAnimationModifierIndex(UAnimationModifier* Instance, int32 Direction);
+	UE_API void AddAnimationModifier(UAnimationModifier* Instance);
+	UE_API void RemoveAnimationModifierInstance(UAnimationModifier* Instance);
+	UE_API void ChangeAnimationModifierIndex(UAnimationModifier* Instance, int32 Direction);
 private:
-	void RemoveInvalidModifiers();
+	UE_API void RemoveInvalidModifiers();
 protected:
 	UPROPERTY()
 	TArray<TObjectPtr<UAnimationModifier>> AnimationModifierInstances;
@@ -76,3 +78,5 @@ protected:
 	UPROPERTY()
 	TMap<FSoftObjectPath, TObjectPtr<UAnimationModifier>> AppliedModifiers;
 };
+
+#undef UE_API
diff --git a/Engine/Source/Editor/AssetDefinition/Public/AssetDefinition.h b/Engine/Source/Editor/AssetDefinition/Public/AssetDefinition.h
--- a/Engine/Source/Editor/AssetDefinition/Public/AssetDefinition.h
+++ b/Engine/Source/Editor/AssetDefinition/Public/AssetDefinition.h
@@ -13,6 +13,8 @@
 
 #include "AssetDefinition.generated.h"
 
+#define UE_API ASSETDEFINITION_API
+
 class IToolkitHost;
 class UThumbnailInfo;
 struct FAssetDisplayInfo;
@@ -267,29 +269,29 @@ struct FAssetOpenSupportArgs
  * These are just some common asset categories.  You're not at all limited to these, and can register an "Advanced"
  * category with the IAssetTools::RegisterAdvancedAssetCategory.
  */
-struct ASSETDEFINITION_API EAssetCategoryPaths
+struct EAssetCategoryPaths
 {
 	// This category is special, "Basic" assets appear at the very top level and are not placed into any submenu.
 	// Arguably the basic category should not exist and should instead be user configurable on what they feel should be
 	// top level assets.
-	static FAssetCategoryPath Basic;
+	static UE_API FAssetCategoryPath Basic;
 	
-	static FAssetCategoryPath Animation;
-	static FAssetCategoryPath Audio;
-	static FAssetCategoryPath Blueprint;
-	static FAssetCategoryPath Cinematics;
-	static FAssetCategoryPath Foliage;
-	static FAssetCategoryPath FX;
-	static FAssetCategoryPath Gameplay;
-	static FAssetCategoryPath AI;
-	static FAssetCategoryPath Input;
-	static FAssetCategoryPath Material;
-	static FAssetCategoryPath Media;
-	static FAssetCategoryPath Misc;
-	static FAssetCategoryPath Physics;
-	static FAssetCategoryPath Texture;
-	static FAssetCategoryPath UI;
-	static FAssetCategoryPath World;
+	static UE_API FAssetCategoryPath Animation;
+	static UE_API FAssetCategoryPath Audio;
+	static UE_API FAssetCategoryPath Blueprint;
+	static UE_API FAssetCategoryPath Cinematics;
+	static UE_API FAssetCategoryPath Foliage;
+	static UE_API FAssetCategoryPath FX;
+	static UE_API FAssetCategoryPath Gameplay;
+	static UE_API FAssetCategoryPath AI;
+	static UE_API FAssetCategoryPath Input;
+	static UE_API FAssetCategoryPath Material;
+	static UE_API FAssetCategoryPath Media;
+	static UE_API FAssetCategoryPath Misc;
+	static UE_API FAssetCategoryPath Physics;
+	static UE_API FAssetCategoryPath Texture;
+	static UE_API FAssetCategoryPath UI;
+	static UE_API FAssetCategoryPath World;
 };
 
 struct FAssetOpenSupport
@@ -448,16 +450,16 @@ public:
  * Blueprints provide a filter for Blueprint Class, but they also provide the filters for Blueprint Interface, Blueprint Macro Library and Blueprint Function library,
  * which are all UBlueprint assets, but differ based on Asset Tag data.
  */
-UCLASS(Abstract)
-class ASSETDEFINITION_API UAssetDefinition : public UObject
+UCLASS(MinimalAPI, Abstract)
+class UAssetDefinition : public UObject
 {
 	GENERATED_BODY()
 
 public:
-	UAssetDefinition();
+	UE_API UAssetDefinition();
 
 	//Begin UObject
-	virtual void PostCDOContruct() override;
+	UE_API virtual void PostCDOContruct() override;
 	//End UObject
 	
 public:
@@ -483,7 +485,7 @@ public:
 	virtual FText GetAssetDescription(const FAssetData& AssetData) const { return FText::GetEmpty(); }
 
 	/** Gets a list of categories this asset is in, these categories are used to help organize */
-	virtual TConstArrayView<FAssetCategoryPath> GetAssetCategories() const;
+	UE_API virtual TConstArrayView<FAssetCategoryPath> GetAssetCategories() const;
 	
 public:
 	// Common Operations
@@ -551,7 +553,7 @@ public:
 	}
 
 protected:
-	virtual void BuildFilters(TArray<FAssetFilterData>& OutFilters) const;
+	UE_API virtual void BuildFilters(TArray<FAssetFilterData>& OutFilters) const;
 
 public:
 	// Extras
@@ -567,10 +569,10 @@ public:
 	 * @param InArgs The asset data of the assets we want the source files from and in which format we want the file path to be.
 	 * @param SourceFileFunc A function that is called for each source file found. The call back must return true to continue the enumeration.
 	 */
-	virtual EAssetCommandResult GetSourceFiles(const FAssetSourceFilesArgs& InArgs, TFunctionRef<bool(const FAssetSourceFilesResult& InSourceFile)> SourceFileFunc) const;
+	UE_API virtual EAssetCommandResult GetSourceFiles(const FAssetSourceFilesArgs& InArgs, TFunctionRef<bool(const FAssetSourceFilesResult& InSourceFile)> SourceFileFunc) const;
 
 	UE_DEPRECATED(5.3, "This override will be removed because it doesn't account that the resolution of a relative path to its absolute path may varies based on the asset implementation.")
-	virtual EAssetCommandResult GetSourceFiles(const FAssetData& InAsset, TFunctionRef<void(const FAssetImportInfo& AssetImportData)> SourceFileFunc) const;
+	UE_API virtual EAssetCommandResult GetSourceFiles(const FAssetData& InAsset, TFunctionRef<void(const FAssetImportInfo& AssetImportData)> SourceFileFunc) const;
 
 	// Diffing Assets
 	virtual EAssetCommandResult PerformAssetDiff(const FAssetDiffArgs& DiffArgs) const
@@ -646,7 +648,7 @@ public:
 	//    Please continue to use EAssetCommandResult and FAssetSupportResponse, for those kinds of commands.
 
 protected:
-	virtual bool CanRegisterStatically() const;
+	UE_API virtual bool CanRegisterStatically() const;
 
 protected:
 	EIncludeClassInFilter IncludeClassInFilter = EIncludeClassInFilter::IfClassIsNotAbstract;
@@ -655,4 +657,6 @@ private:
 	mutable TSharedPtr<FAssetFilterDataCache> FilterCache;
 
 	friend class UAssetDefinitionRegistry;
-};
\ No newline at end of file
+};
+
+#undef UE_API
diff --git a/Engine/Source/Editor/AssetDefinition/Public/AssetDefinitionRegistry.h b/Engine/Source/Editor/AssetDefinition/Public/AssetDefinitionRegistry.h
--- a/Engine/Source/Editor/AssetDefinition/Public/AssetDefinitionRegistry.h
+++ b/Engine/Source/Editor/AssetDefinition/Public/AssetDefinitionRegistry.h
@@ -8,29 +8,31 @@
 
 #include "AssetDefinitionRegistry.generated.h"
 
+#define UE_API ASSETDEFINITION_API
+
 class UAssetDefinition;
 
 DECLARE_MULTICAST_DELEGATE_OneParam(FOnAssetDefinitionRegistryVersionChange, UAssetDefinitionRegistry*);
 
-UCLASS(config=Editor)
-class ASSETDEFINITION_API UAssetDefinitionRegistry : public UObject
+UCLASS(MinimalAPI, config=Editor)
+class UAssetDefinitionRegistry : public UObject
 {
 	GENERATED_BODY()
 
 public:	
-	static UAssetDefinitionRegistry* Get();
+	static UE_API UAssetDefinitionRegistry* Get();
 
-	UAssetDefinitionRegistry();
+	UE_API UAssetDefinitionRegistry();
 	
-	virtual void BeginDestroy() override;
+	UE_API virtual void BeginDestroy() override;
 
-	const UAssetDefinition* GetAssetDefinitionForAsset(const FAssetData& Asset) const;
-	const UAssetDefinition* GetAssetDefinitionForClass(const UClass* Class) const;
+	UE_API const UAssetDefinition* GetAssetDefinitionForAsset(const FAssetData& Asset) const;
+	UE_API const UAssetDefinition* GetAssetDefinitionForClass(const UClass* Class) const;
 
 	// Gets the current version of the AssetDefinitions.   Version is updated whenever an AssetDefinition is Registered/Unregistered
-	uint64 GetAssetDefinitionVersion() const;
-	TArray<TObjectPtr<UAssetDefinition>> GetAllAssetDefinitions() const;
-	TArray<TSoftClassPtr<UObject>> GetAllRegisteredAssetClasses() const;
+	UE_API uint64 GetAssetDefinitionVersion() const;
+	UE_API TArray<TObjectPtr<UAssetDefinition>> GetAllAssetDefinitions() const;
+	UE_API TArray<TSoftClassPtr<UObject>> GetAllRegisteredAssetClasses() const;
 
 	/**
 	 * Normally UAssetDefinitionRegistry are registered automatically by their CDO.  The only reason you need to do this is if
@@ -38,21 +40,21 @@ public:
 	 * to be able to create wrappers for the to be replaced IAssetTypeActions, that you can access AssetDefinition
 	 * versions of any IAssetType making the upgrade easier.
 	 */
-	void RegisterAssetDefinition(UAssetDefinition* AssetDefinition);
+	UE_API void RegisterAssetDefinition(UAssetDefinition* AssetDefinition);
 
-	void UnregisterAssetDefinition(UAssetDefinition* AssetDefinition);
+	UE_API void UnregisterAssetDefinition(UAssetDefinition* AssetDefinition);
 
 	/**
 	 * Called when the AssetDefinitionRegistry's version has changed.
 	 */
-	FOnAssetDefinitionRegistryVersionChange& OnAssetDefinitionRegistryVersionChange();
+	UE_API FOnAssetDefinitionRegistryVersionChange& OnAssetDefinitionRegistryVersionChange();
 	
 private:
-	void RegisterTickerForVersionNotification();
-	bool TickVersionNotification(float);
+	UE_API void RegisterTickerForVersionNotification();
+	UE_API bool TickVersionNotification(float);
 	
-	static UAssetDefinitionRegistry* Singleton;
-	static bool bHasShutDown;
+	static UE_API UAssetDefinitionRegistry* Singleton;
+	static UE_API bool bHasShutDown;
 
 	UPROPERTY()
 	TMap<TSoftClassPtr<UObject>, TObjectPtr<UAssetDefinition>> AssetDefinitions;
@@ -62,4 +64,6 @@ private:
 	FTickerDelegate TickerDelegate;
 	FTSTicker::FDelegateHandle TickerDelegateHandle;
 	FOnAssetDefinitionRegistryVersionChange OnAssetDefinitionRegistryVersionChangeDelegate;
-};
\ No newline at end of file
+};
+
+#undef UE_API
diff --git a/Engine/Source/Editor/AssetDefinition/Public/AssetStatusAssetDataInfoProvider.h b/Engine/Source/Editor/AssetDefinition/Public/AssetStatusAssetDataInfoProvider.h
--- a/Engine/Source/Editor/AssetDefinition/Public/AssetStatusAssetDataInfoProvider.h
+++ b/Engine/Source/Editor/AssetDefinition/Public/AssetStatusAssetDataInfoProvider.h
@@ -1,23 +1,27 @@
-﻿// Copyright Epic Games, Inc. All Rights Reserved.
+// Copyright Epic Games, Inc. All Rights Reserved.
 
 #pragma once
 
 #include "IAssetStatusInfoProvider.h"
 #include "AssetRegistry/AssetData.h"
 
-class ASSETDEFINITION_API FAssetStatusAssetDataInfoProvider : public IAssetStatusInfoProvider
+#define UE_API ASSETDEFINITION_API
+
+class FAssetStatusAssetDataInfoProvider : public IAssetStatusInfoProvider
 {
 public:
 	FAssetStatusAssetDataInfoProvider(FAssetData InAssetData)
 		: AssetData(InAssetData)
 	{}
 
-	virtual UPackage* FindPackage() const override;
+	UE_API virtual UPackage* FindPackage() const override;
 
-	virtual FString TryGetFilename() const override;
+	UE_API virtual FString TryGetFilename() const override;
 
-	virtual FAssetData TryGetAssetData() const override;
+	UE_API virtual FAssetData TryGetAssetData() const override;
 
 private:
 	FAssetData AssetData;
 };
+
+#undef UE_API
diff --git a/Engine/Source/Editor/AssetDefinition/Public/Misc/AssetCategoryPath.h b/Engine/Source/Editor/AssetDefinition/Public/Misc/AssetCategoryPath.h
--- a/Engine/Source/Editor/AssetDefinition/Public/Misc/AssetCategoryPath.h
+++ b/Engine/Source/Editor/AssetDefinition/Public/Misc/AssetCategoryPath.h
@@ -4,17 +4,19 @@
 
 #include "CoreMinimal.h"
 
+#define UE_API ASSETDEFINITION_API
+
 /**
  * The asset category path is how we know how to build menus around assets.  For example, Basic is generally the ones
  * we expose at the top level, where as everything else is a category with a pull out menu, and the subcategory would
  * be where it gets placed in a submenu inside of there.
  */
-struct ASSETDEFINITION_API FAssetCategoryPath
+struct FAssetCategoryPath
 {
-	FAssetCategoryPath(const FText& InCategory);
-	FAssetCategoryPath(const FText& InCategory, const FText& SubCategory);
-	FAssetCategoryPath(const FAssetCategoryPath& InCategory, const FText& InSubCategory);
-	FAssetCategoryPath(TConstArrayView<FText> InCategoryPath);
+	UE_API FAssetCategoryPath(const FText& InCategory);
+	UE_API FAssetCategoryPath(const FText& InCategory, const FText& SubCategory);
+	UE_API FAssetCategoryPath(const FAssetCategoryPath& InCategory, const FText& InSubCategory);
+	UE_API FAssetCategoryPath(TConstArrayView<FText> InCategoryPath);
 
 	FName GetCategory() const { return CategoryPath[0].Key; }
 	FText GetCategoryText() const { return CategoryPath[0].Value; }
@@ -24,11 +26,13 @@ struct ASSETDEFINITION_API FAssetCategoryPath
 	FName GetSubCategory() const { return HasSubCategory() ? CategoryPath[1].Key : NAME_None; }
 	FText GetSubCategoryText() const { return HasSubCategory() ? CategoryPath[1].Value : FText::GetEmpty(); }
 
-	void GetSubCategories(TArray<FName>& SubCategories) const;
-	void GetSubCategoriesText(TArray<FText>& SubCategories) const;
+	UE_API void GetSubCategories(TArray<FName>& SubCategories) const;
+	UE_API void GetSubCategoriesText(TArray<FText>& SubCategories) const;
 	
 	FAssetCategoryPath operator / (const FText& SubCategory) const { return FAssetCategoryPath(*this, SubCategory); }
 	
 private:
 	TArray<TPair<FName, FText>> CategoryPath;
-};
\ No newline at end of file
+};
+
+#undef UE_API
diff --git a/Engine/Source/Editor/AssetTagsEditor/Public/SAssetTagItem.h b/Engine/Source/Editor/AssetTagsEditor/Public/SAssetTagItem.h
--- a/Engine/Source/Editor/AssetTagsEditor/Public/SAssetTagItem.h
+++ b/Engine/Source/Editor/AssetTagsEditor/Public/SAssetTagItem.h
@@ -23,6 +23,8 @@
 #include "Widgets/Views/STableRow.h"
 #include "Widgets/Views/STableViewBase.h"
 
+#define UE_API ASSETTAGSEDITOR_API
+
 class SWidget;
 struct FSlateBrush;
 
@@ -97,7 +99,7 @@ private:
 };
 
 /** A single asset tag item in the tags and collections tree. */
-class ASSETTAGSEDITOR_API SAssetTagItem : public SCompoundWidget
+class SAssetTagItem : public SCompoundWidget
 {
 public:
 	SLATE_BEGIN_ARGS(SAssetTagItem)
@@ -153,7 +155,7 @@ public:
 	SLATE_END_ARGS()
 
 	/** Constructs this widget with InArgs */
-	void Construct(const FArguments& InArgs);
+	UE_API void Construct(const FArguments& InArgs);
 
 	void RequestRename()
 	{
@@ -216,3 +218,5 @@ private:
 	/** The widget to display the name of the asset tag item and allows for renaming */
 	TSharedPtr<SInlineEditableTextBlock> InlineRenameWidget;
 };
+
+#undef UE_API
diff --git a/Engine/Source/Editor/AudioEditor/Classes/AudioEditorSettings.h b/Engine/Source/Editor/AudioEditor/Classes/AudioEditorSettings.h
--- a/Engine/Source/Editor/AudioEditor/Classes/AudioEditorSettings.h
+++ b/Engine/Source/Editor/AudioEditor/Classes/AudioEditorSettings.h
@@ -15,6 +15,8 @@
 
 #include "AudioEditorSettings.generated.h"
 
+#define UE_API AUDIOEDITOR_API
+
 namespace Audio
 {
 	struct FAudioPlatformDeviceInfo;
@@ -47,8 +49,8 @@ struct FAudioPlatform
 	}
 };
 
-UCLASS(config = EditorSettings, meta = (DisplayName = "Audio"))
-class AUDIOEDITOR_API UAudioEditorSettings : public UDeveloperSettings
+UCLASS(MinimalAPI, config = EditorSettings, meta = (DisplayName = "Audio"))
+class UAudioEditorSettings : public UDeveloperSettings
 {
 	GENERATED_BODY()
 
@@ -92,7 +94,7 @@ public:
 	TArray<FAudioPlatform> AudioMixerPlatforms;
 	
 	/** Set and apply, whether audio attenuation is used for non-game worlds*/
-	void SetUseAudioAttenuation(bool bInUseAudioAttenuation);
+	UE_API void SetUseAudioAttenuation(bool bInUseAudioAttenuation);
 
 	/** Is audio attenuation used for non-game worlds*/
 	bool IsUsingAudioAttenuation() const { return bUseAudioAttenuation; };
@@ -102,17 +104,17 @@ public:
 	//~ End UDeveloperSettings
 
 	//~ Begin UObject
-	virtual void PostEditChangeProperty(FPropertyChangedEvent& PropertyChangedEvent) override;
-	virtual void PostInitProperties() override;
+	UE_API virtual void PostEditChangeProperty(FPropertyChangedEvent& PropertyChangedEvent) override;
+	UE_API virtual void PostInitProperties() override;
 	//~ End UObject
 
 private:
 
 	/** Apply non-game world attenuation setting for all AudioDevices.*/
-	void ApplyAttenuationForAllAudioDevices();
+	UE_API void ApplyAttenuationForAllAudioDevices();
 
 	/** Apply non-game world attenuation setting for AudioDevice with given ID.*/
-	void ApplyAttenuationForAudioDevice(Audio::FDeviceId InDeviceID);
+	UE_API void ApplyAttenuationForAudioDevice(Audio::FDeviceId InDeviceID);
 };
 
 class FAudioOutputDeviceCustomization : public IDetailCustomization
@@ -158,3 +160,5 @@ private:
 	TSharedPtr<FString> GetCurrentBackendName() const;
 	void SetCurrentBackendName(const FString& InBackendName);
 };
+
+#undef UE_API
diff --git a/Engine/Source/Editor/AudioEditor/Classes/Factories/SoundSubmixFactory.h b/Engine/Source/Editor/AudioEditor/Classes/Factories/SoundSubmixFactory.h
--- a/Engine/Source/Editor/AudioEditor/Classes/Factories/SoundSubmixFactory.h
+++ b/Engine/Source/Editor/AudioEditor/Classes/Factories/SoundSubmixFactory.h
@@ -12,50 +12,54 @@
 
 #include "SoundSubmixFactory.generated.h"
 
+#define UE_API AUDIOEDITOR_API
+
 class FFeedbackContext;
 class UClass;
 class UObject;
 
-UCLASS(hidecategories=Object)
-class AUDIOEDITOR_API USoundSubmixFactory : public UFactory
+UCLASS(MinimalAPI, hidecategories=Object)
+class USoundSubmixFactory : public UFactory
 {
 	GENERATED_UCLASS_BODY()
 
 	//~ Begin UFactory Interface
-	virtual UObject* FactoryCreateNew(UClass* Class,UObject* InParent,FName Name,EObjectFlags Flags,UObject* Context,FFeedbackContext* Warn) override;
-	virtual bool CanCreateNew() const override;
+	UE_API virtual UObject* FactoryCreateNew(UClass* Class,UObject* InParent,FName Name,EObjectFlags Flags,UObject* Context,FFeedbackContext* Warn) override;
+	UE_API virtual bool CanCreateNew() const override;
 	//~ Begin UFactory Interface	
 };
 
-UCLASS(hidecategories = Object)
-class AUDIOEDITOR_API USoundfieldSubmixFactory: public UFactory
+UCLASS(MinimalAPI, hidecategories = Object)
+class USoundfieldSubmixFactory: public UFactory
 {
 	GENERATED_UCLASS_BODY()
 
 	//~ Begin UFactory Interface
-	virtual UObject* FactoryCreateNew(UClass* Class, UObject* InParent, FName Name, EObjectFlags Flags, UObject* Context, FFeedbackContext* Warn) override;
-	virtual bool CanCreateNew() const override;
+	UE_API virtual UObject* FactoryCreateNew(UClass* Class, UObject* InParent, FName Name, EObjectFlags Flags, UObject* Context, FFeedbackContext* Warn) override;
+	UE_API virtual bool CanCreateNew() const override;
 	//~ Begin UFactory Interface	
 };
 
-UCLASS(hidecategories = Object)
-class AUDIOEDITOR_API UEndpointSubmixFactory : public UFactory
+UCLASS(MinimalAPI, hidecategories = Object)
+class UEndpointSubmixFactory : public UFactory
 {
 	GENERATED_UCLASS_BODY()
 
 	//~ Begin UFactory Interface
-	virtual UObject* FactoryCreateNew(UClass* Class, UObject* InParent, FName Name, EObjectFlags Flags, UObject* Context, FFeedbackContext* Warn) override;
-	virtual bool CanCreateNew() const override;
+	UE_API virtual UObject* FactoryCreateNew(UClass* Class, UObject* InParent, FName Name, EObjectFlags Flags, UObject* Context, FFeedbackContext* Warn) override;
+	UE_API virtual bool CanCreateNew() const override;
 	//~ Begin UFactory Interface	
 };
 
-UCLASS(hidecategories = Object)
-class AUDIOEDITOR_API USoundfieldEndpointSubmixFactory : public UFactory
+UCLASS(MinimalAPI, hidecategories = Object)
+class USoundfieldEndpointSubmixFactory : public UFactory
 {
 	GENERATED_UCLASS_BODY()
 
 	//~ Begin UFactory Interface
-	virtual UObject* FactoryCreateNew(UClass* Class, UObject* InParent, FName Name, EObjectFlags Flags, UObject* Context, FFeedbackContext* Warn) override;
-	virtual bool CanCreateNew() const override;
+	UE_API virtual UObject* FactoryCreateNew(UClass* Class, UObject* InParent, FName Name, EObjectFlags Flags, UObject* Context, FFeedbackContext* Warn) override;
+	UE_API virtual bool CanCreateNew() const override;
 	//~ Begin UFactory Interface	
 };
+
+#undef UE_API
diff --git a/Engine/Source/Editor/AudioEditor/Classes/SoundClassGraph/SoundClassGraphSchema.h b/Engine/Source/Editor/AudioEditor/Classes/SoundClassGraph/SoundClassGraphSchema.h
--- a/Engine/Source/Editor/AudioEditor/Classes/SoundClassGraph/SoundClassGraphSchema.h
+++ b/Engine/Source/Editor/AudioEditor/Classes/SoundClassGraph/SoundClassGraphSchema.h
@@ -18,6 +18,8 @@
 
 #include "SoundClassGraphSchema.generated.h"
 
+#define UE_API AUDIOEDITOR_API
+
 class UEdGraph;
 class UEdGraphNode;
 class UEdGraphPin;
@@ -26,7 +28,7 @@ struct FEdGraphPinType;
 
 /** Action to add a node to the graph */
 USTRUCT()
-struct AUDIOEDITOR_API FSoundClassGraphSchemaAction_NewNode : public FEdGraphSchemaAction
+struct FSoundClassGraphSchemaAction_NewNode : public FEdGraphSchemaAction
 {
 	GENERATED_USTRUCT_BODY();
 
@@ -45,7 +47,7 @@ struct AUDIOEDITOR_API FSoundClassGraphSchemaAction_NewNode : public FEdGraphSch
 
 	//~ Begin FEdGraphSchemaAction Interface
 	virtual FName GetTypeId() const override { return StaticGetTypeId(); } 
-	virtual UEdGraphNode* PerformAction(class UEdGraph* ParentGraph, UEdGraphPin* FromPin, const FVector2f& Location, bool bSelectNewNode = true) override;
+	UE_API virtual UEdGraphNode* PerformAction(class UEdGraph* ParentGraph, UEdGraphPin* FromPin, const FVector2f& Location, bool bSelectNewNode = true) override;
 	//~ End FEdGraphSchemaAction Interface
 
 	/** Name for the new SoundClass */
@@ -75,3 +77,4 @@ class USoundClassGraphSchema : public UEdGraphSchema
 	//~ End EdGraphSchema Interface
 };
 
+#undef UE_API
diff --git a/Engine/Source/Editor/AudioEditor/Classes/SoundCueGraph/SoundCueGraphSchema.h b/Engine/Source/Editor/AudioEditor/Classes/SoundCueGraph/SoundCueGraphSchema.h
--- a/Engine/Source/Editor/AudioEditor/Classes/SoundCueGraph/SoundCueGraphSchema.h
+++ b/Engine/Source/Editor/AudioEditor/Classes/SoundCueGraph/SoundCueGraphSchema.h
@@ -15,6 +15,8 @@
 
 #include "SoundCueGraphSchema.generated.h"
 
+#define UE_API AUDIOEDITOR_API
+
 class FString;
 class UClass;
 class UEdGraph;
@@ -27,7 +29,7 @@ struct FEdGraphPinType;
 
 /** Action to add a node to the graph */
 USTRUCT()
-struct AUDIOEDITOR_API FSoundCueGraphSchemaAction_NewNode : public FEdGraphSchemaAction
+struct FSoundCueGraphSchemaAction_NewNode : public FEdGraphSchemaAction
 {
 	GENERATED_USTRUCT_BODY();
 
@@ -48,17 +50,17 @@ struct AUDIOEDITOR_API FSoundCueGraphSchemaAction_NewNode : public FEdGraphSchem
 
 	//~ Begin FEdGraphSchemaAction Interface
 	using FEdGraphSchemaAction::PerformAction; // Prevent hiding of deprecated base class function with FVector2D
-	virtual UEdGraphNode* PerformAction(class UEdGraph* ParentGraph, UEdGraphPin* FromPin, const FVector2f& Location, bool bSelectNewNode = true) override;
+	UE_API virtual UEdGraphNode* PerformAction(class UEdGraph* ParentGraph, UEdGraphPin* FromPin, const FVector2f& Location, bool bSelectNewNode = true) override;
 	//~ End FEdGraphSchemaAction Interface
 
 private:
 	/** Connects new node to output of selected nodes */
-	void ConnectToSelectedNodes(USoundNode* NewNodeclass, UEdGraph* ParentGraph) const;
+	UE_API void ConnectToSelectedNodes(USoundNode* NewNodeclass, UEdGraph* ParentGraph) const;
 };
 
 /** Action to add nodes to the graph based on selected objects*/
 USTRUCT()
-struct AUDIOEDITOR_API FSoundCueGraphSchemaAction_NewFromSelected : public FSoundCueGraphSchemaAction_NewNode
+struct FSoundCueGraphSchemaAction_NewFromSelected : public FSoundCueGraphSchemaAction_NewNode
 {
 	GENERATED_USTRUCT_BODY();
 
@@ -72,13 +74,13 @@ struct AUDIOEDITOR_API FSoundCueGraphSchemaAction_NewFromSelected : public FSoun
 
 	//~ Begin FEdGraphSchemaAction Interface
 	using FEdGraphSchemaAction::PerformAction; // Prevent hiding of deprecated base class function with FVector2D
-	virtual UEdGraphNode* PerformAction(class UEdGraph* ParentGraph, UEdGraphPin* FromPin, const FVector2f& Location, bool bSelectNewNode = true) override;
+	UE_API virtual UEdGraphNode* PerformAction(class UEdGraph* ParentGraph, UEdGraphPin* FromPin, const FVector2f& Location, bool bSelectNewNode = true) override;
 	//~ End FEdGraphSchemaAction Interface
 };
 
 /** Action to create new comment */
 USTRUCT()
-struct AUDIOEDITOR_API FSoundCueGraphSchemaAction_NewComment : public FEdGraphSchemaAction
+struct FSoundCueGraphSchemaAction_NewComment : public FEdGraphSchemaAction
 {
 	GENERATED_USTRUCT_BODY();
 
@@ -92,13 +94,13 @@ struct AUDIOEDITOR_API FSoundCueGraphSchemaAction_NewComment : public FEdGraphSc
 
 	//~ Begin FEdGraphSchemaAction Interface
 	using FEdGraphSchemaAction::PerformAction; // Prevent hiding of deprecated base class function with FVector2D
-	virtual UEdGraphNode* PerformAction(class UEdGraph* ParentGraph, UEdGraphPin* FromPin, const FVector2f& Location, bool bSelectNewNode = true) override;
+	UE_API virtual UEdGraphNode* PerformAction(class UEdGraph* ParentGraph, UEdGraphPin* FromPin, const FVector2f& Location, bool bSelectNewNode = true) override;
 	//~ End FEdGraphSchemaAction Interface
 };
 
 /** Action to paste clipboard contents into the graph */
 USTRUCT()
-struct AUDIOEDITOR_API FSoundCueGraphSchemaAction_Paste : public FEdGraphSchemaAction
+struct FSoundCueGraphSchemaAction_Paste : public FEdGraphSchemaAction
 {
 	GENERATED_USTRUCT_BODY();
 
@@ -112,7 +114,7 @@ struct AUDIOEDITOR_API FSoundCueGraphSchemaAction_Paste : public FEdGraphSchemaA
 
 	//~ Begin FEdGraphSchemaAction Interface
 	using FEdGraphSchemaAction::PerformAction; // Prevent hiding of deprecated base class function with FVector2D
-	virtual UEdGraphNode* PerformAction(class UEdGraph* ParentGraph, UEdGraphPin* FromPin, const FVector2f& Location, bool bSelectNewNode = true) override;
+	UE_API virtual UEdGraphNode* PerformAction(class UEdGraph* ParentGraph, UEdGraphPin* FromPin, const FVector2f& Location, bool bSelectNewNode = true) override;
 	//~ End FEdGraphSchemaAction Interface
 };
 
@@ -160,3 +162,4 @@ private:
 	TArray<TSubclassOf<USoundNode>> AllowedSoundNodes;
 };
 
+#undef UE_API
diff --git a/Engine/Source/Editor/AudioEditor/Classes/SoundSubmixGraph/SoundSubmixGraphSchema.h b/Engine/Source/Editor/AudioEditor/Classes/SoundSubmixGraph/SoundSubmixGraphSchema.h
--- a/Engine/Source/Editor/AudioEditor/Classes/SoundSubmixGraph/SoundSubmixGraphSchema.h
+++ b/Engine/Source/Editor/AudioEditor/Classes/SoundSubmixGraph/SoundSubmixGraphSchema.h
@@ -20,6 +20,8 @@
 
 #include "SoundSubmixGraphSchema.generated.h"
 
+#define UE_API AUDIOEDITOR_API
+
 class FMenuBuilder;
 class FSlateRect;
 class FSlateWindowElementList;
@@ -91,7 +93,7 @@ public:
 
 /** Action to add a node to the graph */
 USTRUCT()
-struct AUDIOEDITOR_API FSoundSubmixGraphSchemaAction_NewNode : public FEdGraphSchemaAction
+struct FSoundSubmixGraphSchemaAction_NewNode : public FEdGraphSchemaAction
 {
 	GENERATED_USTRUCT_BODY();
 
@@ -110,7 +112,7 @@ struct AUDIOEDITOR_API FSoundSubmixGraphSchemaAction_NewNode : public FEdGraphSc
 
 	//~ Begin FEdGraphSchemaAction Interface
 	virtual FName GetTypeId() const override { return StaticGetTypeId(); } 
-	virtual UEdGraphNode* PerformAction(class UEdGraph* ParentGraph, UEdGraphPin* FromPin, const FVector2f& Location, bool bSelectNewNode = true) override;
+	UE_API virtual UEdGraphNode* PerformAction(class UEdGraph* ParentGraph, UEdGraphPin* FromPin, const FVector2f& Location, bool bSelectNewNode = true) override;
 	//~ End FEdGraphSchemaAction Interface
 
 	/** Name for the new SoundSubmix */
@@ -141,3 +143,4 @@ class USoundSubmixGraphSchema : public UEdGraphSchema
 	//~ End EdGraphSchema Interface
 };
 
+#undef UE_API
diff --git a/Engine/Source/Editor/AudioEditor/Public/AssetTypeActions/AssetTypeActions_SoundBase.h b/Engine/Source/Editor/AudioEditor/Public/AssetTypeActions/AssetTypeActions_SoundBase.h
--- a/Engine/Source/Editor/AudioEditor/Public/AssetTypeActions/AssetTypeActions_SoundBase.h
+++ b/Engine/Source/Editor/AudioEditor/Public/AssetTypeActions/AssetTypeActions_SoundBase.h
@@ -14,6 +14,8 @@
 #include "Templates/SharedPointer.h"
 #include "UObject/WeakObjectPtrTemplates.h"
 
+#define UE_API AUDIOEDITOR_API
+
 class SWidget;
 class UClass;
 class UObject;
@@ -22,57 +24,59 @@ struct FAssetData;
 
 class
 // UE_DEPRECATED(5.2, "The AssetDefinition system is replacing AssetTypeActions and UAssetDefinition_SoundBase replaced this.  Please see the Conversion Guide in AssetDefinition.h")
-AUDIOEDITOR_API
+
 FAssetTypeActions_SoundBase : public FAssetTypeActions_Base
 {
 public:
 	// IAssetTypeActions Implementation
 	virtual FText GetName() const override { return NSLOCTEXT("AssetTypeActions", "AssetTypeActions_SoundBase", "Sound Base"); }
 	virtual FColor GetTypeColor() const override { return FColor(97, 85, 212); }
-	virtual UClass* GetSupportedClass() const override;
-	virtual bool AssetsActivatedOverride(const TArray<UObject*>& InObjects, EAssetTypeActivationMethod::Type ActivationType) override;
+	UE_API virtual UClass* GetSupportedClass() const override;
+	UE_API virtual bool AssetsActivatedOverride(const TArray<UObject*>& InObjects, EAssetTypeActivationMethod::Type ActivationType) override;
 	virtual uint32 GetCategories() override { return EAssetTypeCategories::Sounds; }
 	virtual bool CanFilter() override { return false; }
-	virtual TSharedPtr<SWidget> GetThumbnailOverlay(const FAssetData& AssetData) const override;
+	UE_API virtual TSharedPtr<SWidget> GetThumbnailOverlay(const FAssetData& AssetData) const override;
 
 protected:
 	/** Plays the specified sound wave */
-	void PlaySound(USoundBase* Sound) const;
+	UE_API void PlaySound(USoundBase* Sound) const;
 
 	/** Stops any currently playing sounds */
-	void StopSound() const;
+	UE_API void StopSound() const;
 
 	/** Return true if the specified sound is playing */
-	bool IsSoundPlaying(USoundBase* Sound) const;
+	UE_API bool IsSoundPlaying(USoundBase* Sound) const;
 
 	/** Return true if the specified asset's sound is playing */
-	bool IsSoundPlaying(const FAssetData& AssetData) const;
+	UE_API bool IsSoundPlaying(const FAssetData& AssetData) const;
 
 private:
 	/** Handler for when PlaySound is selected */
-	void ExecutePlaySound(TArray<TWeakObjectPtr<USoundBase>> Objects) const;
+	UE_API void ExecutePlaySound(TArray<TWeakObjectPtr<USoundBase>> Objects) const;
 
 	/** Handler for when StopSound is selected */
-	void ExecuteStopSound(TArray<TWeakObjectPtr<USoundBase>> Objects) const;
+	UE_API void ExecuteStopSound(TArray<TWeakObjectPtr<USoundBase>> Objects) const;
 
 	/** Returns true if only one sound is selected to play */
-	bool CanExecutePlayCommand(TArray<TWeakObjectPtr<USoundBase>> Objects) const;
+	UE_API bool CanExecutePlayCommand(TArray<TWeakObjectPtr<USoundBase>> Objects) const;
 
 	/** Handler for Mute is selected  */
-	void ExecuteMuteSound(TArray<TWeakObjectPtr<USoundBase>> Objects) const;
+	UE_API void ExecuteMuteSound(TArray<TWeakObjectPtr<USoundBase>> Objects) const;
 	
 	/** Handler for Solo is selected  */
-	void ExecuteSoloSound(TArray<TWeakObjectPtr<USoundBase>> Objects) const;
+	UE_API void ExecuteSoloSound(TArray<TWeakObjectPtr<USoundBase>> Objects) const;
 
 	/** Returns true if the mute state is set.  */
-	bool IsActionCheckedMute(TArray<TWeakObjectPtr<USoundBase>> Objects) const;
+	UE_API bool IsActionCheckedMute(TArray<TWeakObjectPtr<USoundBase>> Objects) const;
 
 	/** Returns true if the solo state is set.  */
-	bool IsActionCheckedSolo(TArray<TWeakObjectPtr<USoundBase>> Objects) const;
+	UE_API bool IsActionCheckedSolo(TArray<TWeakObjectPtr<USoundBase>> Objects) const;
 
 	/** Returns true if its possible to mute a sound */
-	bool CanExecuteMuteCommand(TArray<TWeakObjectPtr<USoundBase>> Objects) const;
+	UE_API bool CanExecuteMuteCommand(TArray<TWeakObjectPtr<USoundBase>> Objects) const;
 	
 	/** Returns true if its possible to solo a sound */
-	bool CanExecuteSoloCommand(TArray<TWeakObjectPtr<USoundBase>> Objects) const;
+	UE_API bool CanExecuteSoloCommand(TArray<TWeakObjectPtr<USoundBase>> Objects) const;
 };
+
+#undef UE_API
diff --git a/Engine/Source/Editor/AudioEditor/Public/SoundClassEditorUtilities.h b/Engine/Source/Editor/AudioEditor/Public/SoundClassEditorUtilities.h
--- a/Engine/Source/Editor/AudioEditor/Public/SoundClassEditorUtilities.h
+++ b/Engine/Source/Editor/AudioEditor/Public/SoundClassEditorUtilities.h
@@ -6,8 +6,10 @@
 #include "CoreMinimal.h"
 #include "Math/Vector2D.h"
 
+#define UE_API AUDIOEDITOR_API
 
-class AUDIOEDITOR_API FSoundClassEditorUtilities
+
+class FSoundClassEditorUtilities
 {
 public:
 
@@ -18,8 +20,10 @@ public:
 	 * @param	FromPin		Pin that was dragged to create sound class
 	 * @param	Location	Location for new sound class
 	 */
-	static void CreateSoundClass(const class UEdGraph* Graph, class UEdGraphPin* FromPin, const FVector2D& Location, FString Name);
+	static UE_API void CreateSoundClass(const class UEdGraph* Graph, class UEdGraphPin* FromPin, const FVector2D& Location, FString Name);
 
 private:
 	FSoundClassEditorUtilities() {}
 };
+
+#undef UE_API
