Control Rig: Avoid schematic panel ticking when window is not visible
#rnx

diff --git a/Engine/Plugins/Runtime/RigVM/Source/RigVMEditor/Private/Editor/RigVMEditor.cpp b/Engine/Plugins/Runtime/RigVM/Source/RigVMEditor/Private/Editor/RigVMEditor.cpp
--- a/Engine/Plugins/Runtime/RigVM/Source/RigVMEditor/Private/Editor/RigVMEditor.cpp
+++ b/Engine/Plugins/Runtime/RigVM/Source/RigVMEditor/Private/Editor/RigVMEditor.cpp
@@ -4521,7 +4521,8 @@ FString FRigVMEditorBase::GetCustomDebugObjectLabel(UObject* ObjectBeingDebugged
 
 	if (ObjectBeingDebugged == GetRigVMHost())
 	{
-		return TEXT("Editor Preview");
+		static const FString EditorPreviewStr = TEXT("Editor Preview");
+		return EditorPreviewStr;
 	}
 
 	if (const AActor* ParentActor = ObjectBeingDebugged->GetTypedOuter<AActor>())
diff --git a/Engine/Source/Editor/AnimationEditorWidgets/Private/SchematicGraphPanel/SSchematicGraphPanel.cpp b/Engine/Source/Editor/AnimationEditorWidgets/Private/SchematicGraphPanel/SSchematicGraphPanel.cpp
--- a/Engine/Source/Editor/AnimationEditorWidgets/Private/SchematicGraphPanel/SSchematicGraphPanel.cpp
+++ b/Engine/Source/Editor/AnimationEditorWidgets/Private/SchematicGraphPanel/SSchematicGraphPanel.cpp
@@ -1048,8 +1048,30 @@ TStatId SSchematicGraphPanel::GetStatId() const
 	RETURN_QUICK_DECLARE_CYCLE_STAT(SSchematicGraphPanel, STATGROUP_Tickables);
 }
 
+static bool IsInMinimizedOrBackgrounded(const TSharedPtr<SWidget>& Widget)
+{
+	if (!Widget.IsValid())
+	{
+		return true;
+	}
+
+	TSharedPtr<SWindow> Window = FSlateApplication::Get().FindWidgetWindow(Widget.ToSharedRef());
+	if (!Window.IsValid()
+		|| Window->IsWindowMinimized())
+	{
+		return true;
+	}
+
+	return false;
+}
+
 void SSchematicGraphPanel::Tick(float DeltaTime)
 {
+	if (IsInMinimizedOrBackgrounded(AsShared()))
+	{
+		return;
+	}
+	
 	DPIScale.Reset();
 	NodeLabelOffset = FVector2d::ZeroVector;
 	
@@ -1133,6 +1155,11 @@ void SSchematicGraphPanel::Tick(float DeltaTime)
 
 void SSchematicGraphPanel::Tick(const FGeometry& AllottedGeometry, const double InCurrentTime, const float InDeltaTime)
 {
+	if (IsInMinimizedOrBackgrounded(AsShared()))
+	{
+		return;
+	}
+	
 	SNodePanel::Tick(AllottedGeometry, InCurrentTime, InDeltaTime);
 	Tick(InDeltaTime);
 }
